Add stubs for skeleton template framework

diff --git a/Engine/Plugins/Experimental/Animation/SkeletonTemplateFramework/SkeletonTemplateFramework.uplugin b/Engine/Plugins/Experimental/Animation/SkeletonTemplateFramework/SkeletonTemplateFramework.uplugin
new file mode 100644
--- /dev/null
+++ b/Engine/Plugins/Experimental/Animation/SkeletonTemplateFramework/SkeletonTemplateFramework.uplugin
@@ -0,0 +1,32 @@
+{
+	"FileVersion": 3,
+	"Version": 1,
+	"VersionName": "0.1",
+	"FriendlyName": "Skeleton Template Framework",
+	"Description": "",
+	"Category": "Animation",
+	"CreatedBy": "Epic Games, Inc.",
+	"CreatedByURL": "https://epicgames.com",
+	"DocsURL": "",
+	"MarketplaceURL": "",
+	"SupportURL": "",
+	"EnabledByDefault": false,
+	"CanContainContent": true,
+	"IsBetaVersion": false,
+	"IsExperimentalVersion": true,
+	"Installed": false,
+	"Modules": [
+		{
+			"Name": "SkeletonTemplateFrameworkRuntime",
+			"Type": "Runtime",
+			"LoadingPhase": "Default"
+		},
+		{
+			"Name": "SkeletonTemplateFrameworkEditor",
+			"Type": "Editor",
+			"LoadingPhase": "Default"
+		}
+	],
+	"Plugins": [
+	]
+}
diff --git a/Engine/Plugins/Experimental/Animation/SkeletonTemplateFramework/Source/Editor/Private/SkeletonTemplateFramework/Module.cpp b/Engine/Plugins/Experimental/Animation/SkeletonTemplateFramework/Source/Editor/Private/SkeletonTemplateFramework/Module.cpp
new file mode 100644
--- /dev/null
+++ b/Engine/Plugins/Experimental/Animation/SkeletonTemplateFramework/Source/Editor/Private/SkeletonTemplateFramework/Module.cpp
@@ -0,0 +1,20 @@
+// Copyright Epic Games, Inc. All Rights Reserved.
+
+#include "Modules/ModuleManager.h"
+
+namespace UE::Anim::STF
+{
+	class FEditorModule : public IModuleInterface
+	{
+	public:
+		virtual void StartupModule() override
+		{
+		}
+
+		virtual void ShutdownModule() override
+		{
+		}
+	};
+}
+
+IMPLEMENT_MODULE(UE::Anim::STF::FEditorModule, SkeletonTemplateFrameworkEditor)
diff --git a/Engine/Plugins/Experimental/Animation/SkeletonTemplateFramework/Source/Editor/Private/SkeletonTemplateFramework/SkeletonBindingAssetDefinition.cpp b/Engine/Plugins/Experimental/Animation/SkeletonTemplateFramework/Source/Editor/Private/SkeletonTemplateFramework/SkeletonBindingAssetDefinition.cpp
new file mode 100644
--- /dev/null
+++ b/Engine/Plugins/Experimental/Animation/SkeletonTemplateFramework/Source/Editor/Private/SkeletonTemplateFramework/SkeletonBindingAssetDefinition.cpp
@@ -0,0 +1,41 @@
+// Copyright Epic Games, Inc. All Rights Reserved.
+
+#include "SkeletonTemplateFramework/SkeletonBindingAssetDefinition.h"
+#include "SkeletonTemplateFramework/SkeletonBinding.h"
+#include "Toolkits/SimpleAssetEditor.h"
+
+#define LOCTEXT_NAMESPACE "SkeletonTemplateFramework"
+
+FText UAssetDefinition_SkeletonBinding::GetAssetDisplayName() const
+{
+	return LOCTEXT("SkeletonBinding", "Skeleton Binding");
+}
+
+FLinearColor UAssetDefinition_SkeletonBinding::GetAssetColor() const
+{
+	return FLinearColor(FColor::Purple);
+}
+
+TSoftClassPtr<UObject> UAssetDefinition_SkeletonBinding::GetAssetClass() const
+{
+	return USkeletonBinding::StaticClass();
+}
+
+EAssetCommandResult UAssetDefinition_SkeletonBinding::OpenAssets(const FAssetOpenArgs& OpenArgs) const
+{
+	if (OpenArgs.OpenMethod == EAssetOpenMethod::Edit)
+	{
+		TArray<UObject*> Assets = OpenArgs.LoadObjects<UObject>();
+		FSimpleAssetEditor::CreateEditor(EToolkitMode::Standalone, OpenArgs.ToolkitHost, Assets);
+	}
+
+	return EAssetCommandResult::Handled;
+}
+
+TConstArrayView<FAssetCategoryPath> UAssetDefinition_SkeletonBinding::GetAssetCategories() const
+{
+	static const auto Categories = { FAssetCategoryPath(EAssetCategoryPaths::Animation, LOCTEXT("AnimNextSubMenu", "AnimNext")) };
+	return Categories;
+}
+
+#undef LOCTEXT_NAMESPACE
diff --git a/Engine/Plugins/Experimental/Animation/SkeletonTemplateFramework/Source/Editor/Private/SkeletonTemplateFramework/SkeletonBindingAssetDefinition.h b/Engine/Plugins/Experimental/Animation/SkeletonTemplateFramework/Source/Editor/Private/SkeletonTemplateFramework/SkeletonBindingAssetDefinition.h
new file mode 100644
--- /dev/null
+++ b/Engine/Plugins/Experimental/Animation/SkeletonTemplateFramework/Source/Editor/Private/SkeletonTemplateFramework/SkeletonBindingAssetDefinition.h
@@ -0,0 +1,22 @@
+// Copyright Epic Games, Inc. All Rights Reserved.
+
+#pragma once
+
+#include "AssetDefinitionDefault.h"
+
+#include "SkeletonBindingAssetDefinition.generated.h"
+
+UCLASS()
+class UAssetDefinition_SkeletonBinding : public UAssetDefinitionDefault
+{
+	GENERATED_BODY()
+
+public:
+	// Begin UAssetDefinition
+	virtual FText GetAssetDisplayName() const override;
+	virtual FLinearColor GetAssetColor() const override;
+	virtual TSoftClassPtr<UObject> GetAssetClass() const override;
+	virtual EAssetCommandResult OpenAssets(const FAssetOpenArgs& OpenArgs) const override;
+	virtual TConstArrayView<FAssetCategoryPath> GetAssetCategories() const override;
+	// End UAssetDefinition
+};
diff --git a/Engine/Plugins/Experimental/Animation/SkeletonTemplateFramework/Source/Editor/Private/SkeletonTemplateFramework/SkeletonBindingFactory.cpp b/Engine/Plugins/Experimental/Animation/SkeletonTemplateFramework/Source/Editor/Private/SkeletonTemplateFramework/SkeletonBindingFactory.cpp
new file mode 100644
--- /dev/null
+++ b/Engine/Plugins/Experimental/Animation/SkeletonTemplateFramework/Source/Editor/Private/SkeletonTemplateFramework/SkeletonBindingFactory.cpp
@@ -0,0 +1,30 @@
+// Copyright Epic Games, Inc. All Rights Reserved.
+
+#include "SkeletonTemplateFramework/SkeletonBindingFactory.h"
+#include "SkeletonTemplateFramework/SkeletonBinding.h"
+#include "UObject/Package.h"
+
+USkeletonBindingFactory::USkeletonBindingFactory()
+{
+	bCreateNew = true;
+	bEditAfterNew = true;
+	SupportedClass = USkeletonBinding::StaticClass();
+}
+
+bool USkeletonBindingFactory::ConfigureProperties()
+{
+	return true;
+}
+
+UObject* USkeletonBindingFactory::FactoryCreateNew(UClass* Class, UObject* InParent, FName Name, EObjectFlags Flags, UObject* Context, FFeedbackContext* Warn, FName CallingContext)
+{
+	EObjectFlags FlagsToUse = Flags | RF_Public | RF_Standalone | RF_Transactional | RF_LoadCompleted;
+	if (InParent == GetTransientPackage())
+	{
+		FlagsToUse &= ~RF_Standalone;
+	}
+
+	USkeletonBinding* NewBinding = NewObject<USkeletonBinding>(InParent, Class, Name, FlagsToUse);
+
+	return NewBinding;
+}
diff --git a/Engine/Plugins/Experimental/Animation/SkeletonTemplateFramework/Source/Editor/Private/SkeletonTemplateFramework/SkeletonBindingFactory.h b/Engine/Plugins/Experimental/Animation/SkeletonTemplateFramework/Source/Editor/Private/SkeletonTemplateFramework/SkeletonBindingFactory.h
new file mode 100644
--- /dev/null
+++ b/Engine/Plugins/Experimental/Animation/SkeletonTemplateFramework/Source/Editor/Private/SkeletonTemplateFramework/SkeletonBindingFactory.h
@@ -0,0 +1,23 @@
+// Copyright Epic Games, Inc. All Rights Reserved.
+
+#pragma once
+
+#include "CoreMinimal.h"
+#include "Factories/Factory.h"
+
+#include "SkeletonBindingFactory.generated.h"
+
+UCLASS()
+class USkeletonBindingFactory : public UFactory
+{
+	GENERATED_BODY()
+
+	USkeletonBindingFactory();
+
+	virtual bool ConfigureProperties() override;
+	virtual UObject* FactoryCreateNew(UClass* Class, UObject* InParent, FName Name, EObjectFlags Flags, UObject* Context, FFeedbackContext* Warn, FName CallingContext) override;
+	virtual UObject* FactoryCreateNew(UClass* Class, UObject* InParent, FName Name, EObjectFlags Flags, UObject* Context, FFeedbackContext* Warn) override
+	{
+		return FactoryCreateNew(Class, InParent, Name, Flags, Context, Warn, NAME_None);
+	}
+};
diff --git a/Engine/Plugins/Experimental/Animation/SkeletonTemplateFramework/Source/Editor/Private/SkeletonTemplateFramework/SkeletonTemplateAssetDefinition.cpp b/Engine/Plugins/Experimental/Animation/SkeletonTemplateFramework/Source/Editor/Private/SkeletonTemplateFramework/SkeletonTemplateAssetDefinition.cpp
new file mode 100644
--- /dev/null
+++ b/Engine/Plugins/Experimental/Animation/SkeletonTemplateFramework/Source/Editor/Private/SkeletonTemplateFramework/SkeletonTemplateAssetDefinition.cpp
@@ -0,0 +1,41 @@
+// Copyright Epic Games, Inc. All Rights Reserved.
+
+#include "SkeletonTemplateFramework/SkeletonTemplateAssetDefinition.h"
+#include "SkeletonTemplateFramework/SkeletonTemplate.h"
+#include "Toolkits/SimpleAssetEditor.h"
+
+#define LOCTEXT_NAMESPACE "SkeletonTemplateFramework"
+
+FText UAssetDefinition_SkeletonTemplate::GetAssetDisplayName() const
+{
+	return LOCTEXT("SkeletonTemplate", "Skeleton Template");
+}
+
+FLinearColor UAssetDefinition_SkeletonTemplate::GetAssetColor() const
+{
+	return FLinearColor(FColor::Purple);
+}
+
+TSoftClassPtr<UObject> UAssetDefinition_SkeletonTemplate::GetAssetClass() const
+{
+	return USkeletonTemplate::StaticClass();
+}
+
+EAssetCommandResult UAssetDefinition_SkeletonTemplate::OpenAssets(const FAssetOpenArgs& OpenArgs) const
+{
+	if (OpenArgs.OpenMethod == EAssetOpenMethod::Edit)
+	{
+		TArray<UObject*> Assets = OpenArgs.LoadObjects<UObject>();
+		FSimpleAssetEditor::CreateEditor(EToolkitMode::Standalone, OpenArgs.ToolkitHost, Assets);
+	}
+
+	return EAssetCommandResult::Handled;
+}
+
+TConstArrayView<FAssetCategoryPath> UAssetDefinition_SkeletonTemplate::GetAssetCategories() const
+{
+	static const auto Categories = { FAssetCategoryPath(EAssetCategoryPaths::Animation, LOCTEXT("AnimNextSubMenu", "AnimNext")) };
+	return Categories;
+}
+
+#undef LOCTEXT_NAMESPACE
diff --git a/Engine/Plugins/Experimental/Animation/SkeletonTemplateFramework/Source/Editor/Private/SkeletonTemplateFramework/SkeletonTemplateAssetDefinition.h b/Engine/Plugins/Experimental/Animation/SkeletonTemplateFramework/Source/Editor/Private/SkeletonTemplateFramework/SkeletonTemplateAssetDefinition.h
new file mode 100644
--- /dev/null
+++ b/Engine/Plugins/Experimental/Animation/SkeletonTemplateFramework/Source/Editor/Private/SkeletonTemplateFramework/SkeletonTemplateAssetDefinition.h
@@ -0,0 +1,22 @@
+// Copyright Epic Games, Inc. All Rights Reserved.
+
+#pragma once
+
+#include "AssetDefinitionDefault.h"
+
+#include "SkeletonTemplateAssetDefinition.generated.h"
+
+UCLASS()
+class UAssetDefinition_SkeletonTemplate : public UAssetDefinitionDefault
+{
+	GENERATED_BODY()
+
+public:
+	// Begin UAssetDefinition
+	virtual FText GetAssetDisplayName() const override;
+	virtual FLinearColor GetAssetColor() const override;
+	virtual TSoftClassPtr<UObject> GetAssetClass() const override;
+	virtual EAssetCommandResult OpenAssets(const FAssetOpenArgs& OpenArgs) const override;
+	virtual TConstArrayView<FAssetCategoryPath> GetAssetCategories() const override;
+	// End UAssetDefinition
+};
diff --git a/Engine/Plugins/Experimental/Animation/SkeletonTemplateFramework/Source/Editor/Private/SkeletonTemplateFramework/SkeletonTemplateFactory.cpp b/Engine/Plugins/Experimental/Animation/SkeletonTemplateFramework/Source/Editor/Private/SkeletonTemplateFramework/SkeletonTemplateFactory.cpp
new file mode 100644
--- /dev/null
+++ b/Engine/Plugins/Experimental/Animation/SkeletonTemplateFramework/Source/Editor/Private/SkeletonTemplateFramework/SkeletonTemplateFactory.cpp
@@ -0,0 +1,30 @@
+// Copyright Epic Games, Inc. All Rights Reserved.
+
+#include "SkeletonTemplateFramework/SkeletonTemplateFactory.h"
+#include "SkeletonTemplateFramework/SkeletonTemplate.h"
+#include "UObject/Package.h"
+
+USkeletonTemplateFactory::USkeletonTemplateFactory()
+{
+	bCreateNew = true;
+	bEditAfterNew = true;
+	SupportedClass = USkeletonTemplate::StaticClass();
+}
+
+bool USkeletonTemplateFactory::ConfigureProperties()
+{
+	return true;
+}
+
+UObject* USkeletonTemplateFactory::FactoryCreateNew(UClass* Class, UObject* InParent, FName Name, EObjectFlags Flags, UObject* Context, FFeedbackContext* Warn, FName CallingContext)
+{
+	EObjectFlags FlagsToUse = Flags | RF_Public | RF_Standalone | RF_Transactional | RF_LoadCompleted;
+	if (InParent == GetTransientPackage())
+	{
+		FlagsToUse &= ~RF_Standalone;
+	}
+
+	USkeletonTemplate* NewBinding = NewObject<USkeletonTemplate>(InParent, Class, Name, FlagsToUse);
+
+	return NewBinding;
+}
diff --git a/Engine/Plugins/Experimental/Animation/SkeletonTemplateFramework/Source/Editor/Private/SkeletonTemplateFramework/SkeletonTemplateFactory.h b/Engine/Plugins/Experimental/Animation/SkeletonTemplateFramework/Source/Editor/Private/SkeletonTemplateFramework/SkeletonTemplateFactory.h
new file mode 100644
--- /dev/null
+++ b/Engine/Plugins/Experimental/Animation/SkeletonTemplateFramework/Source/Editor/Private/SkeletonTemplateFramework/SkeletonTemplateFactory.h
@@ -0,0 +1,23 @@
+// Copyright Epic Games, Inc. All Rights Reserved.
+
+#pragma once
+
+#include "CoreMinimal.h"
+#include "Factories/Factory.h"
+
+#include "SkeletonTemplateFactory.generated.h"
+
+UCLASS()
+class USkeletonTemplateFactory : public UFactory
+{
+	GENERATED_BODY()
+
+	USkeletonTemplateFactory();
+
+	virtual bool ConfigureProperties() override;
+	virtual UObject* FactoryCreateNew(UClass* Class, UObject* InParent, FName Name, EObjectFlags Flags, UObject* Context, FFeedbackContext* Warn, FName CallingContext) override;
+	virtual UObject* FactoryCreateNew(UClass* Class, UObject* InParent, FName Name, EObjectFlags Flags, UObject* Context, FFeedbackContext* Warn) override
+	{
+		return FactoryCreateNew(Class, InParent, Name, Flags, Context, Warn, NAME_None);
+	}
+};
diff --git a/Engine/Plugins/Experimental/Animation/SkeletonTemplateFramework/Source/Editor/SkeletonTemplateFrameworkEditor.Build.cs b/Engine/Plugins/Experimental/Animation/SkeletonTemplateFramework/Source/Editor/SkeletonTemplateFrameworkEditor.Build.cs
new file mode 100644
--- /dev/null
+++ b/Engine/Plugins/Experimental/Animation/SkeletonTemplateFramework/Source/Editor/SkeletonTemplateFrameworkEditor.Build.cs
@@ -0,0 +1,19 @@
+// Copyright Epic Games, Inc. All Rights Reserved.
+
+using EpicGames.Core;
+using UnrealBuildTool;
+
+public class SkeletonTemplateFrameworkEditor : ModuleRules
+{
+	public SkeletonTemplateFrameworkEditor(ReadOnlyTargetRules Target) : base(Target)
+	{
+		PrivateDependencyModuleNames.AddAll(
+			"AssetDefinition",
+			"Core",
+			"CoreUObject",
+			"Engine",
+			"UnrealEd",
+			"SkeletonTemplateFrameworkRuntime"
+		);
+	}
+}
diff --git a/Engine/Plugins/Experimental/Animation/SkeletonTemplateFramework/Source/Runtime/Private/SkeletonTemplateFramework/Module.cpp b/Engine/Plugins/Experimental/Animation/SkeletonTemplateFramework/Source/Runtime/Private/SkeletonTemplateFramework/Module.cpp
new file mode 100644
--- /dev/null
+++ b/Engine/Plugins/Experimental/Animation/SkeletonTemplateFramework/Source/Runtime/Private/SkeletonTemplateFramework/Module.cpp
@@ -0,0 +1,20 @@
+// Copyright Epic Games, Inc. All Rights Reserved.
+
+#include "Modules/ModuleManager.h"
+
+namespace UE::Anim::STF
+{
+	class FRuntimeModule : public IModuleInterface
+	{
+	public:
+		virtual void StartupModule() override
+		{
+		}
+
+		virtual void ShutdownModule() override
+		{
+		}
+	};
+}
+
+IMPLEMENT_MODULE(UE::Anim::STF::FRuntimeModule, SkeletonTemplateFrameworkRuntime)
diff --git a/Engine/Plugins/Experimental/Animation/SkeletonTemplateFramework/Source/Runtime/Private/SkeletonTemplateFramework/SkeletonBinding.cpp b/Engine/Plugins/Experimental/Animation/SkeletonTemplateFramework/Source/Runtime/Private/SkeletonTemplateFramework/SkeletonBinding.cpp
new file mode 100644
--- /dev/null
+++ b/Engine/Plugins/Experimental/Animation/SkeletonTemplateFramework/Source/Runtime/Private/SkeletonTemplateFramework/SkeletonBinding.cpp
@@ -0,0 +1,10 @@
+// Copyright Epic Games, Inc. All Rights Reserved.
+
+#include "SkeletonTemplateFramework/SkeletonBinding.h"
+
+#include UE_INLINE_GENERATED_CPP_BY_NAME(SkeletonBinding)
+
+USkeletonBinding::USkeletonBinding(const FObjectInitializer& ObjectInitializer)
+	: Super(ObjectInitializer)
+{
+}
diff --git a/Engine/Plugins/Experimental/Animation/SkeletonTemplateFramework/Source/Runtime/Private/SkeletonTemplateFramework/SkeletonTemplate.cpp b/Engine/Plugins/Experimental/Animation/SkeletonTemplateFramework/Source/Runtime/Private/SkeletonTemplateFramework/SkeletonTemplate.cpp
new file mode 100644
--- /dev/null
+++ b/Engine/Plugins/Experimental/Animation/SkeletonTemplateFramework/Source/Runtime/Private/SkeletonTemplateFramework/SkeletonTemplate.cpp
@@ -0,0 +1,10 @@
+// Copyright Epic Games, Inc. All Rights Reserved.
+
+#include "SkeletonTemplateFramework/SkeletonTemplate.h"
+
+#include UE_INLINE_GENERATED_CPP_BY_NAME(SkeletonTemplate)
+
+USkeletonTemplate::USkeletonTemplate(const FObjectInitializer& ObjectInitializer)
+	: Super(ObjectInitializer)
+{
+}
diff --git a/Engine/Plugins/Experimental/Animation/SkeletonTemplateFramework/Source/Runtime/Public/SkeletonTemplateFramework/SkeletonBinding.h b/Engine/Plugins/Experimental/Animation/SkeletonTemplateFramework/Source/Runtime/Public/SkeletonTemplateFramework/SkeletonBinding.h
new file mode 100644
--- /dev/null
+++ b/Engine/Plugins/Experimental/Animation/SkeletonTemplateFramework/Source/Runtime/Public/SkeletonTemplateFramework/SkeletonBinding.h
@@ -0,0 +1,77 @@
+// Copyright Epic Games, Inc. All Rights Reserved.
+
+#pragma once
+
+#include "CoreMinimal.h"
+#include "SkeletonTemplateFramework/SkeletonTemplate.h"
+
+#include "SkeletonBinding.generated.h"
+
+class USkeleton;
+
+// A skeleton attribute binding
+USTRUCT()
+struct FSkeletonAttributeBinding
+{
+	GENERATED_BODY()
+
+	// The name of the attribute in the bound template
+	UPROPERTY(EditAnywhere, Category = "Settings")
+	FName AttributeName;
+
+	// The name of the bone in the bound skeleton
+	UPROPERTY(EditAnywhere, Category = "Settings")
+	FName BoneName;
+};
+
+// A skeleton binding specifies how a skeleton binds to a template. A binding asset specifies
+// which named attributes bones map to and it can override the named sets/mappings from the
+// source template as needed for that skeleton.
+UCLASS(MinimalAPI, BlueprintType)
+class USkeletonBinding : public UObject
+{
+	GENERATED_BODY()
+
+public:
+	USkeletonBinding(const FObjectInitializer& ObjectInitializer);
+
+protected:
+	// The skeleton template to bind
+	// We inherit everything from our parent template, allowing us to override things as needed
+	UPROPERTY(EditAnywhere, Category = "Settings")
+	TObjectPtr<USkeletonTemplate> Template;
+
+	// The skeleton to bind
+	// Bones/curves/etc from our skeleton can be bound to named attributes within the template
+	// Any unbound entries are considered named attributes as well
+	UPROPERTY(EditAnywhere, Category = "Settings")
+	TObjectPtr<USkeleton> Skeleton;
+
+	// The list of bound attributes
+	UPROPERTY(EditAnywhere, Category = "Settings")
+	TArray<FSkeletonAttributeBinding> AttributeBindings;
+
+	// The list of attributes within this binding (includes attributes from the template and any unmapped entries from the skeleton)
+	UPROPERTY(VisibleAnywhere, Category = "Settings")
+	TArray<FSkeletonNamedAttribute> NamedAttributes;
+
+	// The list of attribute sets within this binding (includes sets from the template and any overrides we might have)
+	// TODO: A binding can only override sets not add/remove them
+	UPROPERTY(EditAnywhere, Category = "Settings")
+	TArray<FSkeletonNamedAttributeSet> NamedAttributeSets;
+
+	// The list of attribute mappings within this binding (includes mappings from the template and any overrides we might have)
+	// TODO: A binding can only override mappings not add/remove them
+	UPROPERTY(EditAnywhere, Category = "Settings")
+	TArray<FSkeletonNamedAttributeMapping> NamedAttributeMappings;
+
+	// The cached skeleton GUID, if it differs from the one in the bound skeleton,
+	// we are stale and we need to recompute our effective data
+	UPROPERTY()
+	FGuid CachedSkeletonGuid;
+
+	// The cached template hash value, if the hash differs it means our effective data
+	// is stale and needs to be recomputed
+	UPROPERTY()
+	uint32 CachedParentTemplateHash = 0;
+};
diff --git a/Engine/Plugins/Experimental/Animation/SkeletonTemplateFramework/Source/Runtime/Public/SkeletonTemplateFramework/SkeletonTemplate.h b/Engine/Plugins/Experimental/Animation/SkeletonTemplateFramework/Source/Runtime/Public/SkeletonTemplateFramework/SkeletonTemplate.h
new file mode 100644
--- /dev/null
+++ b/Engine/Plugins/Experimental/Animation/SkeletonTemplateFramework/Source/Runtime/Public/SkeletonTemplateFramework/SkeletonTemplate.h
@@ -0,0 +1,114 @@
+// Copyright Epic Games, Inc. All Rights Reserved.
+
+#pragma once
+
+#include "CoreMinimal.h"
+
+#include "SkeletonTemplate.generated.h"
+
+// A skeleton named attribute
+USTRUCT()
+struct FSkeletonNamedAttribute
+{
+	GENERATED_BODY()
+
+	// The name of the attribute
+	UPROPERTY(EditAnywhere, Category = "Settings")
+	FName Name;
+
+	// The name of the parent attribute (optional)
+	UPROPERTY(EditAnywhere, Category = "Settings")
+	FName ParentName;
+
+	// TODO: Attribute Type
+};
+
+// A skeleton named attribute set
+USTRUCT()
+struct FSkeletonNamedAttributeSet
+{
+	GENERATED_BODY()
+
+	// The name of the attribute set
+	UPROPERTY(EditAnywhere, Category = "Settings")
+	FName Name;
+
+	// The list of attributes within this set
+	UPROPERTY(EditAnywhere, Category = "Settings")
+	TArray<FName> NamedAttributes;
+
+	// TODO: Add a hierarchy table, it maps named attributes to which sub-parts are included
+	// Sub-parts are defined per attribute type (e.g. FTransform has rotation/translation/scale sub-parts)
+};
+
+// A skeleton named attribute mapping
+USTRUCT()
+struct FSkeletonNamedAttributeMapping
+{
+	GENERATED_BODY()
+
+	// The name of the attribute mapping
+	UPROPERTY(EditAnywhere, Category = "Settings")
+	FName Name;
+
+	// The name of the source attribute set to use as keys in our mapping
+	UPROPERTY(EditAnywhere, Category = "Settings")
+	FName SourceAttributeSet;
+
+	// TODO: Mapping type
+
+	// TODO: Add a hierarchy table, it maps named attributes from our source set to values of the mapping type
+};
+
+// A skeleton template specifies a list of attributes that skeletons can contain as well
+// as sets and data mappings. A skeleton binding asset can then map each attribute to
+// a skeleton bone/curve/etc as well as override any set/mapping values inherited from
+// its template.
+UCLASS(MinimalAPI, BlueprintType)
+class USkeletonTemplate : public UObject
+{
+	GENERATED_BODY()
+
+public:
+	USkeletonTemplate(const FObjectInitializer& ObjectInitializer);
+
+protected:
+	// An optional parent template we can derive from
+	// We inherit everything from our parent template, allowing us to override things as needed
+	UPROPERTY(EditAnywhere, Category = "Settings")
+	TObjectPtr<USkeletonTemplate> ParentTemplate;
+
+	// The list of attributes within this template
+	UPROPERTY(EditAnywhere, Category = "Settings")
+	TArray<FSkeletonNamedAttribute> NamedAttributes;
+
+	// The list of attribute sets within this template
+	UPROPERTY(EditAnywhere, Category = "Settings")
+	TArray<FSkeletonNamedAttributeSet> NamedAttributeSets;
+
+	// The list of attribute mappings within this template
+	UPROPERTY(EditAnywhere, Category = "Settings")
+	TArray<FSkeletonNamedAttributeMapping> NamedAttributeMappings;
+
+	// The list of attributes within this template, including the parent template
+	UPROPERTY()
+	TArray<FSkeletonNamedAttribute> EffectiveNamedAttributes;
+
+	// The list of attribute sets within this template, including the parent template
+	UPROPERTY()
+	TArray<FSkeletonNamedAttributeSet> EffectiveNamedAttributeSets;
+
+	// The list of attribute mappings within this template, including the parent template
+	UPROPERTY()
+	TArray<FSkeletonNamedAttributeMapping> EffectiveNamedAttributeMappings;
+
+	// The is a hash of our user authored values that derived templates can rely on to
+	// detect staleness
+	UPROPERTY()
+	uint32 AssetHash = 0;
+
+	// The cached parent template hash value, if the hash differs it means our effective data
+	// is stale and needs to be recomputed from the parent and our user authored values
+	UPROPERTY()
+	uint32 CachedParentTemplateHash = 0;
+};
diff --git a/Engine/Plugins/Experimental/Animation/SkeletonTemplateFramework/Source/Runtime/SkeletonTemplateFrameworkRuntime.Build.cs b/Engine/Plugins/Experimental/Animation/SkeletonTemplateFramework/Source/Runtime/SkeletonTemplateFrameworkRuntime.Build.cs
new file mode 100644
--- /dev/null
+++ b/Engine/Plugins/Experimental/Animation/SkeletonTemplateFramework/Source/Runtime/SkeletonTemplateFrameworkRuntime.Build.cs
@@ -0,0 +1,18 @@
+// Copyright Epic Games, Inc. All Rights Reserved.
+
+using UnrealBuildTool;
+
+public class SkeletonTemplateFrameworkRuntime : ModuleRules
+{
+	public SkeletonTemplateFrameworkRuntime(ReadOnlyTargetRules Target) : base(Target)
+	{
+		PublicDependencyModuleNames.AddRange(
+			new string[]
+			{
+				"Core",
+				"CoreUObject",
+				"Engine",
+			}
+		);
+	}
+}
