Add hierarchy table dependency and clean up stubs

diff --git a/Engine/Plugins/Experimental/Animation/SkeletonTemplateFramework/Source/Runtime/Public/SkeletonTemplateFramework/SkeletonTemplate.h b/Engine/Plugins/Experimental/Animation/SkeletonTemplateFramework/Source/Runtime/Public/SkeletonTemplateFramework/SkeletonTemplate.h
--- a/Engine/Plugins/Experimental/Animation/SkeletonTemplateFramework/Source/Runtime/Public/SkeletonTemplateFramework/SkeletonTemplate.h
+++ b/Engine/Plugins/Experimental/Animation/SkeletonTemplateFramework/Source/Runtime/Public/SkeletonTemplateFramework/SkeletonTemplate.h
@@ -3,6 +3,7 @@
 #pragma once
 
 #include "CoreMinimal.h"
+#include "HierarchyTable.h"
 
 #include "SkeletonTemplate.generated.h"
 
@@ -20,7 +21,9 @@ struct FSkeletonNamedAttribute
 	UPROPERTY(EditAnywhere, Category = "Settings")
 	FName ParentName;
 
-	// TODO: Attribute Type
+	// The type of the attribute
+	UPROPERTY(EditAnywhere, Category = "Settings")
+	TObjectPtr<const UScriptStruct> Type;
 };
 
 // A skeleton named attribute set
@@ -55,9 +58,9 @@ struct FSkeletonNamedAttributeMapping
 	UPROPERTY(EditAnywhere, Category = "Settings")
 	FName SourceAttributeSet;
 
-	// TODO: Mapping type
-
-	// TODO: Add a hierarchy table, it maps named attributes from our source set to values of the mapping type
+	// The hierarchy table that contains the mapping data
+	UPROPERTY(EditAnywhere, Category = "Settings")
+	TObjectPtr<UHierarchyTable> TableData;
 };
 
 // A skeleton template specifies a list of attributes that skeletons can contain as well
@@ -73,10 +76,68 @@ public:
 	USkeletonTemplate(const FObjectInitializer& ObjectInitializer);
 
 protected:
+	// TODO Parent Template Support:
+	// A parent template means a few things:
+	//    - Named Attributes come from the parent, these are immutable: cannot rename, change type, change parent, cannot be removed
+	//    - New named attributes can be added, if the parent template changes and an attribute is removed that the
+	//      current template relied on (e.g. had a child attribute under, a set that used it), we retain it as if we had authored it
+	//    - Named Sets come from the parent, its name is immutable, but its content can change (can add/remove attributes, can add/remove sub-parts)
+	//    - Named sets can be added
+	//    - Named sets from the parent cannot be removed
+	//    - Named Mappings come from the parent, its name is immutable, its type is immutable, but its content can change (can change source set, can change mapped values)
+	//    - Named mappings can be added
+	//    - Named mappings from the parent cannot be removed
+	//    - Named mappings that depend on a modified set that is removed by the parent will force that set to be retained locally
+	//    - Generally speaking, if a template modifies something from the parent or creates new entries that depend on parent values
+	//      then if the parent removes them, we will retain them locally as if they had been authored
+	//      However, if we did not change any of these values and we do not depend on them, when the parent removes them, they are removed here as well
+	//    - When we need to display a value to the user, it can come from:
+	//        - If we modified it locally, we show that value
+	//        - If we haven't modified it, we check if it is overridden by a parent entry, if it is we display that value (e.g. set beneath to X)
+	//        - If we haven't modified it, and no parent entry overrides a value, we display the parent value
+	//        - If we haven't modified it, and we have no parent, we display the type's default value
+	//    - This means that we have:
+	//        - source values from the parent template
+	//        - source values from the local template (this), stored as a delta from the parent template (e.g. only what we modified locally)
+	//        - effective values that combine the parent + local
+	//    - Effective values is what is shown to users in the UI and what is returned from queries, it's an implementation detail
+	//    - As such, we need to retain hashes of the parent values to detect changes so that we may rebuild our effective values
+	//    - We need to make sure we cannot select ourself or another parent within the inheritance chain as our parent template
+	//      as this might create a reference cycle
+
+	// TODO Attribute Selection:
+	// Named Sets refer to Named Attributes but they are just generic FNames
+	// Similar to how AnimNext Variables/Parameters were handled, we need a dropdown of valid attributes
+	// Alternatively, if the Named Sets are stored as a hierarchy table, then it needs to display the same topology as named mappings
+	// The attribute selection will also be needed later when we wish to refer to attributes from external systems (e.g. Anim Graph)
+	// There, a graph will specify the template it uses and we'll restrict the attributes shown in a dropdown
+
+	// TODO: Attribute Type Support:
+	// We wish to support any built-in/user type but not every type known to man
+	// We need a registry system where we specify what types to expose as valid attribute types
+	// We'll need to provide additional information per type (e.g. sub-parts, default value), this
+	// is where we would specify that (e.g. some adapter we register which specifies the type + metadata)
+
+	// TODO: Mapping Type Support:
+	// Similar to Attribute Types, we need a registry to handle mapping types but here it is
+	// to provide supplemental information. It should be possible to put any type in a mapping,
+	// but we need to be able to specify extra metadata. For example, if we want a 'value type'
+	// mapping to store the default value/bind pose of attributes, this meta-type will be the
+	// mapping type, but each mapped entry will have its own type derived from its attribute type
+	// (e.g. FTransform attributes have an FTransform, float attributes have a float)
+	// 
+	// We also need to be able to specify, through native code, built-in mappings. Some mappings
+	// are accessed through native code with hardcoded names (e.g. bind pose, translation retarget options)
+	// and those mappings thus have a name that cannot be edited by the user. It might make sense
+	// in the UI to break both types of mapping down when we add a new one: Built-In and Custom (or some other name)
+	// And so the registry can also be used to specify these.
+	// We have to prevent user specifies names from colliding with built-in ones, this is something we must
+	// enforce and validate on load/cook (or when we display the UI)
+
 	// An optional parent template we can derive from
 	// We inherit everything from our parent template, allowing us to override things as needed
-	UPROPERTY(EditAnywhere, Category = "Settings")
-	TObjectPtr<USkeletonTemplate> ParentTemplate;
+	//UPROPERTY(EditAnywhere, Category = "Settings")
+	//TObjectPtr<USkeletonTemplate> ParentTemplate;
 
 	// The list of attributes within this template
 	UPROPERTY(EditAnywhere, Category = "Settings")
@@ -91,24 +152,24 @@ protected:
 	TArray<FSkeletonNamedAttributeMapping> NamedAttributeMappings;
 
 	// The list of attributes within this template, including the parent template
-	UPROPERTY()
-	TArray<FSkeletonNamedAttribute> EffectiveNamedAttributes;
+	//UPROPERTY(EditAnywhere, Category = "Settings")
+	//TArray<FSkeletonNamedAttribute> EffectiveNamedAttributes;
 
 	// The list of attribute sets within this template, including the parent template
-	UPROPERTY()
-	TArray<FSkeletonNamedAttributeSet> EffectiveNamedAttributeSets;
+	//UPROPERTY(EditAnywhere, Category = "Settings")
+	//TArray<FSkeletonNamedAttributeSet> EffectiveNamedAttributeSets;
 
 	// The list of attribute mappings within this template, including the parent template
-	UPROPERTY()
-	TArray<FSkeletonNamedAttributeMapping> EffectiveNamedAttributeMappings;
+	//UPROPERTY(EditAnywhere, Category = "Settings")
+	//TArray<FSkeletonNamedAttributeMapping> EffectiveNamedAttributeMappings;
 
-	// The is a hash of our user authored values that derived templates can rely on to
+	// The is a hash of our user authored + parent values that derived templates can rely on to
 	// detect staleness
-	UPROPERTY()
-	uint32 AssetHash = 0;
+	//UPROPERTY(VisibleAnywhere, Category = "Internals")
+	//uint32 AssetHash = 0;
 
 	// The cached parent template hash value, if the hash differs it means our effective data
 	// is stale and needs to be recomputed from the parent and our user authored values
-	UPROPERTY()
-	uint32 CachedParentTemplateHash = 0;
+	//UPROPERTY(VisibleAnywhere, Category = "Internals")
+	//uint32 CachedParentTemplateHash = 0;
 };
diff --git a/Engine/Plugins/Experimental/Animation/SkeletonTemplateFramework/Source/Runtime/SkeletonTemplateFrameworkRuntime.Build.cs b/Engine/Plugins/Experimental/Animation/SkeletonTemplateFramework/Source/Runtime/SkeletonTemplateFrameworkRuntime.Build.cs
--- a/Engine/Plugins/Experimental/Animation/SkeletonTemplateFramework/Source/Runtime/SkeletonTemplateFrameworkRuntime.Build.cs
+++ b/Engine/Plugins/Experimental/Animation/SkeletonTemplateFramework/Source/Runtime/SkeletonTemplateFrameworkRuntime.Build.cs
@@ -12,6 +12,7 @@ public class SkeletonTemplateFrameworkRuntime : ModuleRules
 				"Core",
 				"CoreUObject",
 				"Engine",
+				"HierarchyTableRuntime",
 			}
 		);
 	}
