Skeleton Templates - Basic editor layout and widget for editing named attributes

diff --git a/Engine/Plugins/Experimental/Animation/HierarchyTable/Source/Runtime/Public/HierarchyTable.h b/Engine/Plugins/Experimental/Animation/HierarchyTable/Source/Runtime/Public/HierarchyTable.h
--- a/Engine/Plugins/Experimental/Animation/HierarchyTable/Source/Runtime/Public/HierarchyTable.h
+++ b/Engine/Plugins/Experimental/Animation/HierarchyTable/Source/Runtime/Public/HierarchyTable.h
@@ -89,7 +89,7 @@ private:
 /**
  * A general-purpose container asset for storing typed hierarchical data.
  */
-UCLASS()
+UCLASS(EditInlineNew)
 class HIERARCHYTABLERUNTIME_API UHierarchyTable : public UObject
 {
 	GENERATED_BODY()
diff --git a/Engine/Plugins/Experimental/Animation/SkeletonTemplateFramework/Source/Editor/Private/SkeletonTemplateFramework/SkeletonTemplateAssetDefinition.cpp b/Engine/Plugins/Experimental/Animation/SkeletonTemplateFramework/Source/Editor/Private/SkeletonTemplateFramework/SkeletonTemplateAssetDefinition.cpp
--- a/Engine/Plugins/Experimental/Animation/SkeletonTemplateFramework/Source/Editor/Private/SkeletonTemplateFramework/SkeletonTemplateAssetDefinition.cpp
+++ b/Engine/Plugins/Experimental/Animation/SkeletonTemplateFramework/Source/Editor/Private/SkeletonTemplateFramework/SkeletonTemplateAssetDefinition.cpp
@@ -2,6 +2,7 @@
 
 #include "SkeletonTemplateFramework/SkeletonTemplateAssetDefinition.h"
 #include "SkeletonTemplateFramework/SkeletonTemplate.h"
+#include "SkeletonTemplateFramework/SkeletonTemplateEditorToolkit.h"
 #include "Toolkits/SimpleAssetEditor.h"
 
 #define LOCTEXT_NAMESPACE "SkeletonTemplateFramework"
@@ -26,7 +27,7 @@ EAssetCommandResult UAssetDefinition_SkeletonTemplate::OpenAssets(const FAssetOp
 	if (OpenArgs.OpenMethod == EAssetOpenMethod::Edit)
 	{
 		TArray<UObject*> Assets = OpenArgs.LoadObjects<UObject>();
-		FSimpleAssetEditor::CreateEditor(EToolkitMode::Standalone, OpenArgs.ToolkitHost, Assets);
+		MakeShared<FSkeletonTemplateEditorToolkit>()->InitEditor(Assets);
 	}
 
 	return EAssetCommandResult::Handled;
diff --git a/Engine/Plugins/Experimental/Animation/SkeletonTemplateFramework/Source/Editor/Private/SkeletonTemplateFramework/SkeletonTemplateEditorToolkit.cpp b/Engine/Plugins/Experimental/Animation/SkeletonTemplateFramework/Source/Editor/Private/SkeletonTemplateFramework/SkeletonTemplateEditorToolkit.cpp
new file mode 100644
--- /dev/null
+++ b/Engine/Plugins/Experimental/Animation/SkeletonTemplateFramework/Source/Editor/Private/SkeletonTemplateFramework/SkeletonTemplateEditorToolkit.cpp
@@ -0,0 +1,148 @@
+// Copyright Epic Games, Inc. All Rights Reserved.
+
+#include "SkeletonTemplateFramework/SkeletonTemplateEditorToolkit.h"
+
+#include "Modules/ModuleManager.h"
+#include "PropertyEditorModule.h"
+#include "SkeletonTemplateFramework/SkeletonTemplate.h"
+#include "SkeletonTemplateFramework/SkeletonTemplateNamedAttributesEditor.h"
+#include "Widgets/Docking/SDockTab.h"
+
+#define LOCTEXT_NAMESPACE "SkeletonTemplateEditorToolkit"
+
+namespace UE::Anim::STF
+{
+	struct FEditorTabs
+	{
+		static const FName AttributesId;
+		static const FName TableViewId;
+		static const FName DetailsId;
+	};
+}
+
+const FName UE::Anim::STF::FEditorTabs::AttributesId(TEXT("Attributes"));
+const FName UE::Anim::STF::FEditorTabs::TableViewId(TEXT("TableView"));
+const FName UE::Anim::STF::FEditorTabs::DetailsId(TEXT("Details"));
+
+void FSkeletonTemplateEditorToolkit::InitEditor(const TArray<UObject*>& InObjects)
+{
+	using namespace UE::Anim::STF;
+
+	check(!InObjects.IsEmpty())
+	SkeletonTemplate = CastChecked<USkeletonTemplate>(InObjects[0]);
+
+	const TSharedRef<FTabManager::FLayout> Layout = FTabManager::NewLayout("SkeletonTemplateEditorLayout")
+		->AddArea
+		(
+			FTabManager::NewPrimaryArea()->SetOrientation(Orient_Horizontal)
+			->Split
+			(
+				FTabManager::NewStack()
+				->SetSizeCoefficient(0.25f)
+				->AddTab(FEditorTabs::AttributesId, ETabState::OpenedTab)
+			)
+			->Split
+			(
+				FTabManager::NewStack()
+				->SetSizeCoefficient(0.5f)
+				->AddTab(FEditorTabs::TableViewId, ETabState::OpenedTab)
+			)
+			->Split
+			(
+				FTabManager::NewStack()
+				->SetSizeCoefficient(0.25f)
+				->AddTab(FEditorTabs::DetailsId, ETabState::OpenedTab)
+			)
+		);
+
+	FAssetEditorToolkit::InitAssetEditor(EToolkitMode::Standalone, {}, "SkeletonTemplateEditor", Layout, true, true, InObjects);
+}
+
+void FSkeletonTemplateEditorToolkit::RegisterTabSpawners(const TSharedRef<class FTabManager>& InTabManager)
+{
+	using namespace UE::Anim::STF;
+
+	WorkspaceMenuCategory = InTabManager->AddLocalWorkspaceMenuCategory(LOCTEXT("WorkspaceMenuCategory", "Skeleton Template Editor"));
+	TSharedRef<FWorkspaceItem> WorkspaceMenuCategoryRef = WorkspaceMenuCategory.ToSharedRef();
+
+	FAssetEditorToolkit::RegisterTabSpawners(InTabManager);
+	
+	InTabManager->RegisterTabSpawner(FEditorTabs::AttributesId, FOnSpawnTab::CreateSP(this, &FSkeletonTemplateEditorToolkit::SpawnTabAttributes))
+		.SetDisplayName(LOCTEXT("AttributesTabMenu_Description", "Attributes"))
+		.SetTooltipText(LOCTEXT("AttributesTabMenu_ToolTip", "Shows the attributes panel"))
+		.SetGroup(WorkspaceMenuCategoryRef)
+		.SetIcon(FSlateIcon(FAppStyle::GetAppStyleSetName(), "LevelEditor.Tabs.Outliner"));
+	
+	InTabManager->RegisterTabSpawner(FEditorTabs::TableViewId, FOnSpawnTab::CreateSP(this, &FSkeletonTemplateEditorToolkit::SpawnTabTableView))
+		.SetDisplayName(LOCTEXT("TableViewTabMenu_Description", "Table View"))
+		.SetTooltipText(LOCTEXT("TableViewTabMenu_ToolTip", "Shows the table view panel"))
+		.SetGroup(WorkspaceMenuCategoryRef);
+
+	InTabManager->RegisterTabSpawner(FEditorTabs::DetailsId, FOnSpawnTab::CreateSP(this, &FSkeletonTemplateEditorToolkit::SpawnTabDetails))
+		.SetDisplayName(LOCTEXT("DetailsTabMenu_Description", "Details"))
+		.SetTooltipText(LOCTEXT("DetailsTabMenu_Tooltip", "Shows the details panel"))
+		.SetGroup(WorkspaceMenuCategoryRef)
+		.SetIcon(FSlateIcon(FAppStyle::GetAppStyleSetName(), "LevelEditor.Tabs.Details"));
+}
+
+void FSkeletonTemplateEditorToolkit::UnregisterTabSpawners(const TSharedRef<class FTabManager>& InTabManager)
+{
+	using namespace UE::Anim::STF;
+
+	FAssetEditorToolkit::UnregisterTabSpawners(InTabManager);
+
+	InTabManager->UnregisterTabSpawner(FEditorTabs::AttributesId);
+	InTabManager->UnregisterTabSpawner(FEditorTabs::TableViewId);
+	InTabManager->UnregisterTabSpawner(FEditorTabs::DetailsId);
+}
+
+FName FSkeletonTemplateEditorToolkit::GetToolkitFName() const
+{
+	return "SkeletonTemplateEditor";
+}
+
+FText FSkeletonTemplateEditorToolkit::GetBaseToolkitName() const
+{
+	return LOCTEXT("ToolkitName", "Skeleton Template Editor");
+}
+
+FString FSkeletonTemplateEditorToolkit::GetWorldCentricTabPrefix() const
+{
+	return "SkeletonTemplateEditor";
+}
+
+FLinearColor FSkeletonTemplateEditorToolkit::GetWorldCentricTabColorScale() const
+{
+	return FLinearColor::White;
+}
+
+TSharedRef<SDockTab> FSkeletonTemplateEditorToolkit::SpawnTabAttributes(const FSpawnTabArgs& Args)
+{
+	return SNew(SDockTab)
+		[
+			SAssignNew(AttributesTreeView, UE::Anim::STF::SAttributesTreeView, SkeletonTemplate)
+		];
+}
+
+TSharedRef<SDockTab> FSkeletonTemplateEditorToolkit::SpawnTabTableView(const FSpawnTabArgs& Args)
+{
+	return SNew(SDockTab);
+}
+
+TSharedRef<SDockTab> FSkeletonTemplateEditorToolkit::SpawnTabDetails(const FSpawnTabArgs& Args)
+{
+	FPropertyEditorModule& PropertyEditorModule = FModuleManager::GetModuleChecked<FPropertyEditorModule>("PropertyEditor");
+	
+	FDetailsViewArgs DetailsViewArgs;
+	DetailsViewArgs.NameAreaSettings = FDetailsViewArgs::HideNameArea;
+	
+	DetailsView = PropertyEditorModule.CreateDetailView(DetailsViewArgs);
+	DetailsView->SetObjects(TArray<UObject*>{ SkeletonTemplate });
+	
+	return SNew(SDockTab)
+		[
+			DetailsView.ToSharedRef()
+		];
+}
+
+#undef LOCTEXT_NAMESPACE
\ No newline at end of file
diff --git a/Engine/Plugins/Experimental/Animation/SkeletonTemplateFramework/Source/Editor/Private/SkeletonTemplateFramework/SkeletonTemplateEditorToolkit.h b/Engine/Plugins/Experimental/Animation/SkeletonTemplateFramework/Source/Editor/Private/SkeletonTemplateFramework/SkeletonTemplateEditorToolkit.h
new file mode 100644
--- /dev/null
+++ b/Engine/Plugins/Experimental/Animation/SkeletonTemplateFramework/Source/Editor/Private/SkeletonTemplateFramework/SkeletonTemplateEditorToolkit.h
@@ -0,0 +1,38 @@
+// Copyright Epic Games, Inc. All Rights Reserved.
+
+#pragma once
+
+#include "CoreMinimal.h"
+#include "Toolkits/AssetEditorToolkit.h"
+
+class USkeletonTemplate;
+class SDockTab;
+namespace UE::Anim::STF
+{
+	class SAttributesTreeView;
+}
+
+class FSkeletonTemplateEditorToolkit : public FAssetEditorToolkit
+{
+public:
+	void InitEditor(const TArray<UObject*>& InObjects);
+
+	// IToolkit
+	void RegisterTabSpawners(const TSharedRef<class FTabManager>& TabManager) override;
+	void UnregisterTabSpawners(const TSharedRef<class FTabManager>& TabManager) override;
+	FName GetToolkitFName() const override;
+	FText GetBaseToolkitName() const override;
+	FString GetWorldCentricTabPrefix() const override;
+	FLinearColor GetWorldCentricTabColorScale() const override;
+
+private:
+	TSharedRef<SDockTab> SpawnTabAttributes(const FSpawnTabArgs& Args);
+	TSharedRef<SDockTab> SpawnTabTableView(const FSpawnTabArgs& Args);
+	TSharedRef<SDockTab> SpawnTabDetails(const FSpawnTabArgs& Args);
+
+	TObjectPtr<USkeletonTemplate> SkeletonTemplate;
+
+	TSharedPtr<IDetailsView> DetailsView;
+
+	TSharedPtr<UE::Anim::STF::SAttributesTreeView> AttributesTreeView;
+};
\ No newline at end of file
diff --git a/Engine/Plugins/Experimental/Animation/SkeletonTemplateFramework/Source/Editor/Private/SkeletonTemplateFramework/SkeletonTemplateNamedAttributesEditor.cpp b/Engine/Plugins/Experimental/Animation/SkeletonTemplateFramework/Source/Editor/Private/SkeletonTemplateFramework/SkeletonTemplateNamedAttributesEditor.cpp
new file mode 100644
--- /dev/null
+++ b/Engine/Plugins/Experimental/Animation/SkeletonTemplateFramework/Source/Editor/Private/SkeletonTemplateFramework/SkeletonTemplateNamedAttributesEditor.cpp
@@ -0,0 +1,465 @@
+// Copyright Epic Games, Inc. All Rights Reserved.
+
+#include "SkeletonTemplateNamedAttributesEditor.h"
+
+#include "DragAndDrop/DecoratedDragDropOp.h"
+#include "Framework/MultiBox/MultiBoxBuilder.h"
+#include "Modules/ModuleManager.h"
+#include "SkeletonTemplateFramework/SkeletonTemplate.h"
+#include "StructViewerFilter.h"
+#include "StructViewerModule.h"
+#include "Widgets/Input/SComboButton.h"
+#include "Widgets/Text/SInlineEditableTextBlock.h"
+#include "Widgets/Views/STreeView.h"
+
+#define LOCTEXT_NAMESPACE "SkeletonTemplateNamedAttributesEditor"
+
+namespace UE::Anim::STF
+{
+	struct FOutlinerColumns
+	{
+		static const FName NameId;
+		static const FName TypeId;
+	};
+
+	const FName FOutlinerColumns::NameId(TEXT("Name"));
+	const FName FOutlinerColumns::TypeId(TEXT("Type"));
+
+	class FNamedElementDragDropOp : public FDecoratedDragDropOp
+	{
+	public:
+		DRAG_DROP_OPERATOR_TYPE(FNamedElementDragDropOp, FDecoratedDragDropOp)
+
+		static TSharedRef<FNamedElementDragDropOp> New(TWeakPtr<SAttributesTreeView::FTreeItem> InElement)
+		{
+			TSharedRef<FNamedElementDragDropOp> Operation = MakeShared<FNamedElementDragDropOp>();
+			Operation->Element = InElement;
+			Operation->Construct();
+			return Operation;
+		}
+
+		TSharedPtr<SWidget> GetDefaultDecorator() const override
+		{
+			return SNew(SBorder)
+				.Visibility(EVisibility::Visible)
+				.BorderImage(FAppStyle::GetBrush("Menu.Background"))
+				[
+					SNew(STextBlock)
+					.Text(FText::FromName(Element.Pin()->Name))
+				];
+		}
+
+		TWeakPtr<SAttributesTreeView::FTreeItem> Element;
+	};
+
+	class SNamedElementTableRow : public SMultiColumnTableRow<TSharedPtr<SAttributesTreeView::FTreeItem>>
+	{
+	public:
+		DECLARE_DELEGATE_RetVal_TwoParams(bool /* Success */, FOnRenamed, FName /* Old Name */, FName /* New Name */);
+		DECLARE_DELEGATE_TwoParams(FOnReparented, FName /* Attribute Name */, FName /* New Parent Name */);
+		DECLARE_DELEGATE_RetVal_TwoParams(bool /* Success */, FOnTypeSelected, FName /* Attribute Name */, TObjectPtr<const UScriptStruct> /* New Type */);
+	
+		SLATE_BEGIN_ARGS(SNamedElementTableRow) { }
+			SLATE_ARGUMENT(TSharedPtr<SAttributesTreeView::FTreeItem>, NamedElementTreeItem)
+			SLATE_EVENT(FOnRenamed, OnRenamed);
+			SLATE_EVENT(FOnReparented, OnReparented);
+			SLATE_EVENT(FOnTypeSelected, OnTypeSelected);
+		SLATE_END_ARGS()
+
+		FReply OnDragDetected(const FGeometry& MyGeometry, const FPointerEvent& MouseEvent)
+		{
+			if (MouseEvent.IsMouseButtonDown(EKeys::LeftMouseButton))
+			{
+				const TSharedRef<FNamedElementDragDropOp> DragDropOp = FNamedElementDragDropOp::New(NamedElementTreeItem);
+				return FReply::Handled().BeginDragDrop(DragDropOp);
+			}
+
+			return FReply::Unhandled();
+		}
+
+		TOptional<EItemDropZone> OnCanAcceptDrop(const FDragDropEvent& DragDropEvent, EItemDropZone DropZone, TSharedPtr<SAttributesTreeView::FTreeItem> TargetItem)
+		{
+			TOptional<EItemDropZone> ReturnedDropZone;
+	
+			const TSharedPtr<FNamedElementDragDropOp> DragDropOp = DragDropEvent.GetOperationAs<FNamedElementDragDropOp>();
+			if (DragDropOp.IsValid())
+			{
+        		ReturnedDropZone = EItemDropZone::BelowItem;	
+			}
+	
+			return ReturnedDropZone;
+		}
+
+		FReply OnAcceptDrop(const FDragDropEvent& DragDropEvent, EItemDropZone DropZone, TSharedPtr<SAttributesTreeView::FTreeItem> TargetItem)
+		{
+			const TSharedPtr<FNamedElementDragDropOp> DragDropOp = DragDropEvent.GetOperationAs<FNamedElementDragDropOp>();
+			if (!DragDropOp.IsValid())
+			{
+				return FReply::Unhandled();
+			}
+
+			OnReparented.ExecuteIfBound(DragDropOp->Element.Pin()->Name, NamedElementTreeItem->Name);
+
+			return FReply::Handled();
+		}
+
+		void Construct(const FArguments& InArgs, const TSharedRef<STableViewBase>& InOwnerTableView)
+		{
+			NamedElementTreeItem = InArgs._NamedElementTreeItem;
+			OnRenamed = InArgs._OnRenamed;
+			OnReparented = InArgs._OnReparented;
+			OnTypeSelected = InArgs._OnTypeSelected;
+
+			FSuperRowType::Construct(
+				FSuperRowType::FArguments()
+					.OnDragDetected(this, &SNamedElementTableRow::OnDragDetected)
+					.OnCanAcceptDrop(this, &SNamedElementTableRow::OnCanAcceptDrop)
+					.OnAcceptDrop(this, &SNamedElementTableRow::OnAcceptDrop),
+				InOwnerTableView);
+		}
+
+		void OnCommitRename(const FText& InText, ETextCommit::Type CommitInfo)
+		{
+			const FName OldName = NamedElementTreeItem->Name;
+			const FName NewName = FName(InText.ToString());
+
+			if (OnRenamed.IsBound())
+			{
+				const bool bSuccess = OnRenamed.Execute(OldName, NewName);
+				if (bSuccess)
+				{
+					NamedElementTreeItem->Name = NewName;
+				}
+			}
+		}
+
+		void OnStructPicked(const UScriptStruct* InStruct)
+		{
+			if (OnRenamed.IsBound())
+			{
+				const bool bSuccess = OnTypeSelected.Execute(NamedElementTreeItem->Name, InStruct);
+				if (bSuccess)
+				{
+					NamedElementTreeItem->Type = InStruct ? (InStruct->GetFName()) : NAME_None;
+				}
+			}
+			TypeComboButton->SetIsOpen(false);
+		}
+
+		TSharedRef<SWidget> GenerateStructPicker()
+		{
+			FStructViewerInitializationOptions Options;
+			Options.bShowUnloadedStructs = true;
+			Options.bShowNoneOption = true;
+			Options.NameTypeToDisplay = EStructViewerNameTypeToDisplay::DisplayName;
+			Options.DisplayMode = EStructViewerDisplayMode::ListView;
+
+			FOnStructPicked OnPicked(FOnStructPicked::CreateRaw(this, &SNamedElementTableRow::OnStructPicked));
+
+			return FModuleManager::LoadModuleChecked<FStructViewerModule>("StructViewer")
+				.CreateStructViewer(Options, OnPicked);
+		}
+
+		TSharedRef<SWidget> GenerateWidgetForColumn(const FName& ColumnName) override
+		{
+			if (ColumnName == UE::Anim::STF::FOutlinerColumns::NameId)
+			{
+				TSharedPtr<SInlineEditableTextBlock> InlineWidget;
+
+				TSharedRef<SHorizontalBox> HorizontalBox = SNew(SHorizontalBox)
+					+ SHorizontalBox::Slot()
+						.AutoWidth()
+						[
+							SNew(SExpanderArrow, SharedThis(this))
+								.ShouldDrawWires(true)
+						]
+					+ SHorizontalBox::Slot()
+						.AutoWidth()
+						[
+							SAssignNew(InlineWidget, SInlineEditableTextBlock)
+								.Text_Lambda([this]()
+								{
+									return FText::FromName(NamedElementTreeItem->Name);
+								})
+								.OnTextCommitted(this, &SNamedElementTableRow::OnCommitRename)
+						];
+
+				NamedElementTreeItem->OnRenameRequested.BindSP(InlineWidget.Get(), &SInlineEditableTextBlock::EnterEditingMode);
+
+				return HorizontalBox;
+			}
+			else if (ColumnName == UE::Anim::STF::FOutlinerColumns::TypeId)
+			{
+				return SAssignNew(TypeComboButton, SComboButton)
+					.OnGetMenuContent(this, &SNamedElementTableRow::GenerateStructPicker)
+					.ButtonContent()
+					[
+						SNew(STextBlock)
+						.Text_Lambda([this]()
+							{
+								return FText::FromName(NamedElementTreeItem->Type);
+							})
+					];
+			}
+			else
+			{
+				return SNullWidget::NullWidget;
+			}
+		}
+
+	private:
+		TSharedPtr<SAttributesTreeView::FTreeItem> NamedElementTreeItem;
+		FOnRenamed OnRenamed;
+		FOnReparented OnReparented;
+		FOnTypeSelected OnTypeSelected;
+		TSharedPtr<SComboButton> TypeComboButton;
+	};
+
+	void SAttributesTreeView::Construct(const FArguments& InArgs, TObjectPtr<USkeletonTemplate> InSkeletonTemplate)
+	{
+		SkeletonTemplate = InSkeletonTemplate;
+
+		ChildSlot
+		[
+			SAssignNew(TreeView, STreeView<TSharedPtr<FTreeItem>>)
+				.TreeItemsSource(&RootItems)
+				.OnGenerateRow(this, &SAttributesTreeView::TreeView_GenerateItemRow)
+				.OnGetChildren(this, &SAttributesTreeView::TreeView_HandleGetChildren)
+				.OnContextMenuOpening(this, &SAttributesTreeView::TreeView_HandleContextMenuOpening)
+				.OnItemScrolledIntoView(this, &SAttributesTreeView::TreeView_OnItemScrolledIntoView)
+				.HighlightParentNodesForSelection(true)
+				.HeaderRow
+				(
+					SNew(SHeaderRow)
+						+SHeaderRow::Column(UE::Anim::STF::FOutlinerColumns::NameId)
+							.FillWidth(0.5f)
+							.DefaultLabel(LOCTEXT("NameLabel", "Name"))
+						+SHeaderRow::Column(UE::Anim::STF::FOutlinerColumns::TypeId)
+							.FillWidth(0.5f)
+							.DefaultLabel(LOCTEXT("TypeLabel", "Type"))
+				)
+		];
+
+		RegenerateTreeViewItems();
+
+		// Expand all tree items on construction
+		for (TSharedPtr<FTreeItem> TreeItem : GetAllTreeItems())
+		{
+			TreeView->SetItemExpansion(TreeItem, true);
+		}
+	}
+
+	void SAttributesTreeView::TreeView_HandleGetChildren(TSharedPtr<FTreeItem> InItem, TArray<TSharedPtr<FTreeItem>>& OutChildren) const
+	{
+		OutChildren.Append(InItem->Children);
+	}
+
+	void SAttributesTreeView::TreeView_OnItemScrolledIntoView(TSharedPtr<FTreeItem> InItem, const TSharedPtr<ITableRow>& InWidget)
+	{
+		if(DeferredRenameRequest.IsValid())
+		{
+			DeferredRenameRequest->OnRenameRequested.ExecuteIfBound();
+			DeferredRenameRequest.Reset();
+		}
+	}
+
+	TSharedRef<ITableRow> SAttributesTreeView::TreeView_GenerateItemRow(TSharedPtr<SAttributesTreeView::FTreeItem> Item, const TSharedRef<STableViewBase>& OwnerTable)
+	{
+		return SNew(SNamedElementTableRow, OwnerTable)
+			.OnRenamed_Lambda([this](const FName OldName, const FName NewName) -> bool
+				{
+					return SkeletonTemplate->RenameNamedAttribute(OldName, NewName);
+				})
+			.OnReparented_Lambda([this](const FName AttributeName, const FName NewParentName)
+				{
+					SkeletonTemplate->ReparentNamedAttribute(AttributeName, NewParentName);
+					RegenerateTreeViewItems();
+				})
+			.OnTypeSelected_Lambda([this](const FName AttributeName, const TObjectPtr<const UScriptStruct> NewType)
+				{
+					return SkeletonTemplate->SetNamedAttributeType(AttributeName, NewType);
+				})
+			.NamedElementTreeItem(Item);
+	}
+
+	TSharedPtr<SWidget> SAttributesTreeView::TreeView_HandleContextMenuOpening()
+	{
+		FMenuBuilder MenuBuilder(true, nullptr);
+
+		TArray<TSharedPtr<FTreeItem>> SelectedItems;
+		TreeView->GetSelectedItems(SelectedItems);
+
+		FUIAction AddNewAttributeAction = FUIAction(
+			FExecuteAction::CreateSP(this, &SAttributesTreeView::HandleAddAttribute)
+		);
+
+		if (SelectedItems.IsEmpty())
+		{
+			MenuBuilder.AddMenuEntry(
+				LOCTEXT("AddNamedAttribute", "Add Named Attribute"),
+				LOCTEXT("AddNamedAttribute_Tooltip", "Add a new named attribute"),
+				FSlateIcon(FAppStyle::GetAppStyleSetName(), "Icons.Plus"),
+				AddNewAttributeAction);
+		}
+		else
+		{
+			MenuBuilder.AddMenuEntry(
+				LOCTEXT("AddNamedAttribute", "Add Child Named Attribute"),
+				LOCTEXT("AddNamedAttribute_Tooltip", "Add a new child named attribute"),
+				FSlateIcon(FAppStyle::GetAppStyleSetName(), "Icons.Plus"),
+				AddNewAttributeAction);
+
+			MenuBuilder.AddMenuEntry(
+				LOCTEXT("RenameNamedAttribute", "Rename Named Attribute"),
+				LOCTEXT("RenameNamedAttribute_Tooltip", "Renames the selected new child named attribute"),
+				FSlateIcon(FAppStyle::GetAppStyleSetName(), "GenericCommands.Rename"),
+				FUIAction(
+					FExecuteAction::CreateSP(this, &SAttributesTreeView::HandleRenameAttribute)
+				));
+		
+			MenuBuilder.AddMenuEntry(
+				LOCTEXT("DeleteNamedAttribute", "Delete Named Attribute"),
+				LOCTEXT("DeleteNamedAttribute_Tooltip", "Deletes the selected new child named attribute"),
+				FSlateIcon(FAppStyle::GetAppStyleSetName(), "Icons.Delete"),
+				FUIAction(
+					FExecuteAction::CreateSP(this, &SAttributesTreeView::HandleDeleteAttribute)
+				));
+		}
+
+		return MenuBuilder.MakeWidget();
+	}
+
+	TArray<TSharedPtr<UE::Anim::STF::SAttributesTreeView::FTreeItem>> SAttributesTreeView::GetAllTreeItems()
+	{
+		using namespace UE::Anim::STF;
+
+		TArray<TSharedPtr<SAttributesTreeView::FTreeItem>> AllItems;
+		AllItems.Append(RootItems);
+	
+		for (int32 Index = 0; Index < AllItems.Num(); ++Index)
+		{
+			AllItems.Append(AllItems[Index]->Children);
+		}
+
+		return AllItems;
+	}
+
+	void SAttributesTreeView::RegenerateTreeViewItems()
+	{
+		// Make note of all tree items currently expanded
+		TSet<FName> ExpandedAttributeNames;
+		{
+			for (TSharedPtr<FTreeItem> TreeItem : GetAllTreeItems())
+			{
+				if (TreeView->IsItemExpanded(TreeItem))
+				{
+					ExpandedAttributeNames.Add(TreeItem->Name);
+				}
+			}
+		}
+
+		// Rebuild items
+		{
+			RootItems.Reset();
+	
+			TMap<FName, TSharedPtr<FTreeItem>> ItemMap;
+
+			for (const FSkeletonNamedAttribute& NamedAttribute : SkeletonTemplate->GetNamedAttributes())
+			{
+				TSharedPtr<FTreeItem> Item = MakeShared<FTreeItem>(NamedAttribute.Name);
+				Item->Name = NamedAttribute.Name;
+				Item->Type = NamedAttribute.Type ? (NamedAttribute.Type->GetFName()) : NAME_None;
+
+				const TSharedPtr<FTreeItem>* ParentItem = ItemMap.Find(NamedAttribute.ParentName);
+				if (ParentItem)
+				{
+					(*ParentItem)->Children.Add(Item);
+				}
+				else
+				{
+					RootItems.Add(Item);
+				}
+
+				ItemMap.Add(Item->Name, Item);
+			}
+		}
+
+		// Update tree view and restore tree item expanded states
+		{
+			check(TreeView);
+			TreeView->RequestTreeRefresh();
+
+			for (TSharedPtr<FTreeItem> TreeItem : GetAllTreeItems())
+			{
+				if (ExpandedAttributeNames.Contains(TreeItem->Name))
+				{
+					TreeView->SetItemExpansion(TreeItem, true);
+				}
+			}
+		}
+	}
+
+	void SAttributesTreeView::HandleAddAttribute()
+	{
+		using namespace UE::Anim::STF;
+
+		FSkeletonNamedAttribute NewAttribute;
+		NewAttribute.Name = FName(LOCTEXT("NewAttribute", "NewAttribute").ToString());
+
+		TArray<TSharedPtr<SAttributesTreeView::FTreeItem>> SelectedItems;
+		TreeView->GetSelectedItems(SelectedItems);
+	
+		if (!SelectedItems.IsEmpty())
+		{
+			NewAttribute.ParentName = SelectedItems[0]->Name;
+		}
+
+		const bool bSuccess = SkeletonTemplate->AddNamedAttribute(NewAttribute);
+		if (bSuccess)
+		{
+			RegenerateTreeViewItems();
+		}
+
+		for (TSharedPtr<FTreeItem> Item : GetAllTreeItems())
+		{
+			if (Item->Name == NewAttribute.Name)
+			{
+				TreeView->RequestScrollIntoView(Item);
+				DeferredRenameRequest = Item;
+			}
+		}
+	}
+
+	void SAttributesTreeView::HandleRenameAttribute()
+	{
+		using namespace UE::Anim::STF;
+
+		FSkeletonNamedAttribute NewAttribute;
+
+		TArray<TSharedPtr<SAttributesTreeView::FTreeItem>> SelectedItems;
+		TreeView->GetSelectedItems(SelectedItems);
+	
+		check(!SelectedItems.IsEmpty());
+		SelectedItems[0]->OnRenameRequested.ExecuteIfBound();
+	}
+
+	void SAttributesTreeView::HandleDeleteAttribute()
+	{
+		using namespace UE::Anim::STF;
+
+		FSkeletonNamedAttribute NewAttribute;
+
+		TArray<TSharedPtr<FTreeItem>> SelectedItems;
+		TreeView->GetSelectedItems(SelectedItems);
+	
+		check(!SelectedItems.IsEmpty());
+
+		const bool bSuccess = SkeletonTemplate->RemoveNamedAttribute(SelectedItems[0]->Name);
+		if (bSuccess)
+		{
+			RegenerateTreeViewItems();
+		}
+	}
+}
+
+#undef LOCTEXT_NAMESPACE
\ No newline at end of file
diff --git a/Engine/Plugins/Experimental/Animation/SkeletonTemplateFramework/Source/Editor/Private/SkeletonTemplateFramework/SkeletonTemplateNamedAttributesEditor.h b/Engine/Plugins/Experimental/Animation/SkeletonTemplateFramework/Source/Editor/Private/SkeletonTemplateFramework/SkeletonTemplateNamedAttributesEditor.h
new file mode 100644
--- /dev/null
+++ b/Engine/Plugins/Experimental/Animation/SkeletonTemplateFramework/Source/Editor/Private/SkeletonTemplateFramework/SkeletonTemplateNamedAttributesEditor.h
@@ -0,0 +1,62 @@
+// Copyright Epic Games, Inc. All Rights Reserved.
+
+#pragma once
+
+#include "CoreMinimal.h"
+#include "Widgets/SCompoundWidget.h"
+
+class ITableRow;
+class STableViewBase;
+template<typename ItemType> class STreeView;
+class USkeletonTemplate;
+
+namespace UE::Anim::STF
+{
+	class SAttributesTreeView : public SCompoundWidget
+	{
+	public:
+		SLATE_BEGIN_ARGS(SAttributesTreeView) { }
+		SLATE_END_ARGS()
+
+		struct FTreeItem
+		{
+			FName Name;
+			FName Type;
+			TArray<TSharedPtr<FTreeItem>> Children;
+
+			DECLARE_DELEGATE(FOnRenameRequested);
+			FOnRenameRequested OnRenameRequested;
+		};
+
+		void Construct(const FArguments& InArgs, TObjectPtr<USkeletonTemplate> InSkeletonTemplate);
+
+	private:
+		void TreeView_HandleGetChildren(TSharedPtr<FTreeItem> InItem, TArray<TSharedPtr<FTreeItem>>& OutChildren) const;
+		
+		void TreeView_OnItemScrolledIntoView(TSharedPtr<FTreeItem> InItem, const TSharedPtr<ITableRow>& InWidget);
+
+		TSharedRef<ITableRow> TreeView_GenerateItemRow(TSharedPtr<FTreeItem> Item, const TSharedRef<STableViewBase>& OwnerTable);
+		
+		TSharedPtr<SWidget> TreeView_HandleContextMenuOpening();
+
+	private:
+		TArray<TSharedPtr<FTreeItem>> GetAllTreeItems();
+
+		void RegenerateTreeViewItems();
+
+		void HandleAddAttribute();
+
+		void HandleRenameAttribute();
+
+		void HandleDeleteAttribute();
+
+	private:
+		TSharedPtr<STreeView<TSharedPtr<FTreeItem>>> TreeView;
+
+		TArray<TSharedPtr<FTreeItem>> RootItems;
+
+		TSharedPtr<FTreeItem> DeferredRenameRequest;
+
+		TObjectPtr<USkeletonTemplate> SkeletonTemplate;
+	};
+}
\ No newline at end of file
diff --git a/Engine/Plugins/Experimental/Animation/SkeletonTemplateFramework/Source/Editor/SkeletonTemplateFrameworkEditor.Build.cs b/Engine/Plugins/Experimental/Animation/SkeletonTemplateFramework/Source/Editor/SkeletonTemplateFrameworkEditor.Build.cs
--- a/Engine/Plugins/Experimental/Animation/SkeletonTemplateFramework/Source/Editor/SkeletonTemplateFrameworkEditor.Build.cs
+++ b/Engine/Plugins/Experimental/Animation/SkeletonTemplateFramework/Source/Editor/SkeletonTemplateFrameworkEditor.Build.cs
@@ -12,8 +12,11 @@ public class SkeletonTemplateFrameworkEditor : ModuleRules
 			"Core",
 			"CoreUObject",
 			"Engine",
+			"InputCore",
 			"UnrealEd",
-			"SkeletonTemplateFrameworkRuntime"
+			"SkeletonTemplateFrameworkRuntime",
+			"Slate",
+			"SlateCore"
 		);
 	}
 }
diff --git a/Engine/Plugins/Experimental/Animation/SkeletonTemplateFramework/Source/Runtime/Private/SkeletonTemplateFramework/SkeletonTemplate.cpp b/Engine/Plugins/Experimental/Animation/SkeletonTemplateFramework/Source/Runtime/Private/SkeletonTemplateFramework/SkeletonTemplate.cpp
--- a/Engine/Plugins/Experimental/Animation/SkeletonTemplateFramework/Source/Runtime/Private/SkeletonTemplateFramework/SkeletonTemplate.cpp
+++ b/Engine/Plugins/Experimental/Animation/SkeletonTemplateFramework/Source/Runtime/Private/SkeletonTemplateFramework/SkeletonTemplate.cpp
@@ -8,3 +8,228 @@ USkeletonTemplate::USkeletonTemplate(const FObjectInitializer& ObjectInitializer
 	: Super(ObjectInitializer)
 {
 }
+
+const TArray<FSkeletonNamedAttribute>& USkeletonTemplate::GetNamedAttributes() const
+{
+	return NamedAttributes;
+}
+
+bool USkeletonTemplate::AddNamedAttribute(const FSkeletonNamedAttribute& NewNamedAttribute)
+{
+	if (NewNamedAttribute.Name == NAME_None)
+	{
+		return false; // Reserved name for denoting unrooted attributes
+	}
+
+	int32 ExistingAttributeIndex = FindNamedAttributeIndex(NewNamedAttribute.Name);
+	if (ExistingAttributeIndex == INDEX_NONE)
+	{
+		NamedAttributes.Add(NewNamedAttribute);
+		return true;
+	}
+
+	return false;
+}
+
+bool USkeletonTemplate::RenameNamedAttribute(const FName OldAttributeName, const FName NewAttributeName)
+{
+	if (NewAttributeName == NAME_None)
+	{
+		return false; // Reserved name for denoting unrooted attributes
+	}
+
+	int32 AttributeToRenameIndex = FindNamedAttributeIndex(OldAttributeName);
+	if (AttributeToRenameIndex == INDEX_NONE)
+	{
+		return false; // Attribute with old name does not exist
+	}
+
+	int32 ExistingAttributeIndex = FindNamedAttributeIndex(NewAttributeName);
+	if (ExistingAttributeIndex != INDEX_NONE)
+	{
+		return false; // Attribute with this name already exists
+	}
+
+	for (FSkeletonNamedAttribute& NamedAttribute : NamedAttributes)
+	{
+		if (NamedAttribute.ParentName == OldAttributeName)
+		{
+			NamedAttribute.ParentName = NewAttributeName;
+		}
+	}
+	
+	FSkeletonNamedAttribute& NamedAttribute = NamedAttributes[AttributeToRenameIndex];
+	NamedAttribute.Name = NewAttributeName;
+
+	return true;
+}
+
+bool USkeletonTemplate::ReparentNamedAttribute(const FName AttributeName, const FName NewParentName)
+{
+	if (AttributeName == NewParentName)
+	{
+		return false; // Cannot set parent of self
+	}
+
+	FSkeletonNamedAttribute* AttributeToReparent = FindNamedAttribute(AttributeName);
+	if (!AttributeToReparent)
+	{
+		return false; // Attribute does not exist
+	}
+
+	FSkeletonNamedAttribute* NewParentAttribute = FindNamedAttribute(NewParentName);
+	if (!NewParentAttribute)
+	{
+		return false; // Parent does not exist
+	}
+
+	// Update hierarchy to prevent cycles
+	{
+		FSkeletonNamedAttribute* AncestorAttribute = NewParentAttribute;
+		while (AncestorAttribute)
+		{
+			if (AncestorAttribute->ParentName == AttributeName)
+			{
+				AncestorAttribute->ParentName = AttributeToReparent->ParentName;
+				break;
+			}
+			AncestorAttribute = FindNamedAttribute(AncestorAttribute->ParentName);
+		}
+	}
+
+	AttributeToReparent->ParentName = NewParentName;
+	SortNamedAttributes();
+	return true;
+}
+
+bool USkeletonTemplate::RemoveNamedAttribute(const FName AttributeToRemoveName, const ERemoveNamedAttributeBehavior Behavior)
+{
+	FName ParentOfAttributeToRemove = NAME_None;
+	
+	if (Behavior == ERemoveNamedAttributeBehavior::MoveChildrenToParent)
+	{
+		int32 AttributeToRemoveIndex = FindNamedAttributeIndex(AttributeToRemoveName);
+
+		if (AttributeToRemoveIndex != INDEX_NONE)
+		{
+			ParentOfAttributeToRemove = NamedAttributes[AttributeToRemoveIndex].ParentName;
+			NamedAttributes.RemoveAtSwap(AttributeToRemoveIndex, EAllowShrinking::No);
+		}
+		else
+		{
+			return false;
+		}
+	}
+
+	const int32 NumNamedAttributes = NamedAttributes.Num();
+	for (int32 Index = NumNamedAttributes - 1; Index >= 0; --Index)
+	{
+		FSkeletonNamedAttribute& NamedAttribute = NamedAttributes[Index];
+
+		if (NamedAttribute.ParentName == AttributeToRemoveName)
+		{
+			if (Behavior == ERemoveNamedAttributeBehavior::RemoveChildren)
+			{
+				NamedAttributes.RemoveAtSwap(Index, EAllowShrinking::No);
+			}
+			else
+			{
+				NamedAttribute.ParentName = ParentOfAttributeToRemove;
+			}
+		}
+	}
+
+	SortNamedAttributes();
+	return true;
+}
+
+bool USkeletonTemplate::SetNamedAttributeType(const FName AttributeName, const TObjectPtr<const UScriptStruct> NewType)
+{
+	FSkeletonNamedAttribute* NamedAttribute = FindNamedAttribute(AttributeName);
+
+	if (NamedAttribute)
+	{
+		NamedAttribute->Type = NewType;
+		return true;
+	}
+
+	return false;
+}
+
+FSkeletonNamedAttribute* USkeletonTemplate::FindNamedAttribute(const FName AttributeName)
+{
+	return NamedAttributes.FindByPredicate([AttributeName](const FSkeletonNamedAttribute& Candidate)
+		{
+			return Candidate.Name == AttributeName;
+		});
+}
+
+int32 USkeletonTemplate::FindNamedAttributeIndex(const FName AttributeName)
+{
+	return NamedAttributes.IndexOfByPredicate([AttributeName](const FSkeletonNamedAttribute& Candidate)
+		{
+			return Candidate.Name == AttributeName;
+		});
+}
+
+void USkeletonTemplate::SortNamedAttributes()
+{
+	struct Node
+	{
+		FName AttributeName;
+		TArray<FName> Children;
+		int32 InDegree;
+
+		Node(const FName InAttributeName)
+			: AttributeName(InAttributeName)
+			, InDegree(0)
+		{
+		}
+	};
+
+	TMap<FName, Node> Graph;
+
+	for (const FSkeletonNamedAttribute& NamedAttribute : NamedAttributes)
+	{
+		Graph.Add(NamedAttribute.Name, Node(NamedAttribute.Name));
+		Graph.Add(NamedAttribute.ParentName, Node(NamedAttribute.ParentName));
+	}
+
+	for (const FSkeletonNamedAttribute& NamedAttribute : NamedAttributes)
+	{
+		Graph[NamedAttribute.ParentName].Children.Add(NamedAttribute.Name);
+		++Graph[NamedAttribute.Name].InDegree;
+	}
+
+	TArray<FName> NodeQueue;
+	for (const TPair<FName, Node>& AttributeDegree : Graph)
+	{
+		if (AttributeDegree.Value.InDegree == 0)
+		{
+			NodeQueue.Add(AttributeDegree.Key);
+		}
+	}
+
+	TArray<FSkeletonNamedAttribute> SortedNamedAttributes;
+
+	while (!NodeQueue.IsEmpty())
+	{
+		const FName AttributeName = NodeQueue.Pop();
+		if (AttributeName != NAME_None)
+		{
+			SortedNamedAttributes.Add(*FindNamedAttribute(AttributeName));
+		}
+
+		for (const FName& Child : Graph[AttributeName].Children)
+		{
+			--Graph[Child].InDegree;
+			if (Graph[Child].InDegree == 0)
+			{
+				NodeQueue.Add(Child);
+			}
+		}
+	}
+
+	ensureMsgf(NamedAttributes.Num() == SortedNamedAttributes.Num(), TEXT("Cycle found in the named attribute DAG"));
+	NamedAttributes = SortedNamedAttributes;
+}
\ No newline at end of file
diff --git a/Engine/Plugins/Experimental/Animation/SkeletonTemplateFramework/Source/Runtime/Public/SkeletonTemplateFramework/NamedElementTypes.h b/Engine/Plugins/Experimental/Animation/SkeletonTemplateFramework/Source/Runtime/Public/SkeletonTemplateFramework/NamedElementTypes.h
new file mode 100644
--- /dev/null
+++ b/Engine/Plugins/Experimental/Animation/SkeletonTemplateFramework/Source/Runtime/Public/SkeletonTemplateFramework/NamedElementTypes.h
@@ -0,0 +1,16 @@
+// Copyright Epic Games, Inc. All Rights Reserved.
+
+#pragma once
+
+#include "CoreMinimal.h"
+
+#include "NamedElementTypes.generated.h"
+
+USTRUCT()
+struct FSkeletonNamedAttributeType_Transform
+{
+	GENERATED_BODY()
+
+	UPROPERTY()
+	FTransform Value;
+};
\ No newline at end of file
diff --git a/Engine/Plugins/Experimental/Animation/SkeletonTemplateFramework/Source/Runtime/Public/SkeletonTemplateFramework/SkeletonTemplate.h b/Engine/Plugins/Experimental/Animation/SkeletonTemplateFramework/Source/Runtime/Public/SkeletonTemplateFramework/SkeletonTemplate.h
--- a/Engine/Plugins/Experimental/Animation/SkeletonTemplateFramework/Source/Runtime/Public/SkeletonTemplateFramework/SkeletonTemplate.h
+++ b/Engine/Plugins/Experimental/Animation/SkeletonTemplateFramework/Source/Runtime/Public/SkeletonTemplateFramework/SkeletonTemplate.h
@@ -75,7 +75,32 @@ class USkeletonTemplate : public UObject
 public:
 	USkeletonTemplate(const FObjectInitializer& ObjectInitializer);
 
+	// Returns topologically sorted attributes, i.e. parents are defined before they're referenced by children
+	SKELETONTEMPLATEFRAMEWORKRUNTIME_API const TArray<FSkeletonNamedAttribute>& GetNamedAttributes() const;
+
+	SKELETONTEMPLATEFRAMEWORKRUNTIME_API bool AddNamedAttribute(const FSkeletonNamedAttribute& NewNamedAttribute);
+
+	SKELETONTEMPLATEFRAMEWORKRUNTIME_API bool RenameNamedAttribute(const FName OldAttributeName, const FName NewAttributeName);
+
+	SKELETONTEMPLATEFRAMEWORKRUNTIME_API bool ReparentNamedAttribute(const FName AttributeName, const FName NewParentName);
+
+	enum class ERemoveNamedAttributeBehavior
+	{
+		RemoveChildren,
+		MoveChildrenToParent,
+	};
+
+	SKELETONTEMPLATEFRAMEWORKRUNTIME_API bool RemoveNamedAttribute(const FName AttributeName, const ERemoveNamedAttributeBehavior Behavior = ERemoveNamedAttributeBehavior::MoveChildrenToParent);
+
+	SKELETONTEMPLATEFRAMEWORKRUNTIME_API bool SetNamedAttributeType(const FName AttributeName, const TObjectPtr<const UScriptStruct> NewType);
+
 protected:
+	FSkeletonNamedAttribute* FindNamedAttribute(const FName AttributeName);
+
+	int32 FindNamedAttributeIndex(const FName AttributeName);
+
+	void SortNamedAttributes();
+
 	// TODO Parent Template Support:
 	// A parent template means a few things:
 	//    - Named Attributes come from the parent, these are immutable: cannot rename, change type, change parent, cannot be removed
@@ -140,7 +165,7 @@ protected:
 	//TObjectPtr<USkeletonTemplate> ParentTemplate;
 
 	// The list of attributes within this template
-	UPROPERTY(EditAnywhere, Category = "Settings")
+	UPROPERTY(VisibleAnywhere, Category = "Settings")
 	TArray<FSkeletonNamedAttribute> NamedAttributes;
 
 	// The list of attribute sets within this template
