Fix skeleton template named attributes editor compile error

diff --git a/Engine/Plugins/Experimental/Animation/SkeletonTemplateFramework/Source/Editor/Private/SkeletonTemplateFramework/SkeletonTemplateNamedAttributesEditor.cpp b/Engine/Plugins/Experimental/Animation/SkeletonTemplateFramework/Source/Editor/Private/SkeletonTemplateFramework/SkeletonTemplateNamedAttributesEditor.cpp
--- a/Engine/Plugins/Experimental/Animation/SkeletonTemplateFramework/Source/Editor/Private/SkeletonTemplateFramework/SkeletonTemplateNamedAttributesEditor.cpp
+++ b/Engine/Plugins/Experimental/Animation/SkeletonTemplateFramework/Source/Editor/Private/SkeletonTemplateFramework/SkeletonTemplateNamedAttributesEditor.cpp
@@ -366,7 +366,7 @@ namespace UE::Anim::STF
 
 			for (const FSkeletonNamedAttribute& NamedAttribute : SkeletonTemplate->GetNamedAttributes())
 			{
-				TSharedPtr<FTreeItem> Item = MakeShared<FTreeItem>(NamedAttribute.Name);
+				TSharedPtr<FTreeItem> Item = MakeShared<FTreeItem>();
 				Item->Name = NamedAttribute.Name;
 				Item->Type = NamedAttribute.Type ? (NamedAttribute.Type->GetFName()) : NAME_None;
 
