Skeleton Templates - Fix LOCTEXT conflict

diff --git a/Engine/Plugins/Experimental/Animation/SkeletonTemplateFramework/Source/Editor/Private/SkeletonTemplateFramework/SkeletonTemplateNamedAttributesEditor.cpp b/Engine/Plugins/Experimental/Animation/SkeletonTemplateFramework/Source/Editor/Private/SkeletonTemplateFramework/SkeletonTemplateNamedAttributesEditor.cpp
--- a/Engine/Plugins/Experimental/Animation/SkeletonTemplateFramework/Source/Editor/Private/SkeletonTemplateFramework/SkeletonTemplateNamedAttributesEditor.cpp
+++ b/Engine/Plugins/Experimental/Animation/SkeletonTemplateFramework/Source/Editor/Private/SkeletonTemplateFramework/SkeletonTemplateNamedAttributesEditor.cpp
@@ -304,8 +304,8 @@ namespace UE::Anim::STF
 		else
 		{
 			MenuBuilder.AddMenuEntry(
-				LOCTEXT("AddNamedAttribute", "Add Child Named Attribute"),
-				LOCTEXT("AddNamedAttribute_Tooltip", "Add a new child named attribute"),
+				LOCTEXT("AddChildNamedAttribute", "Add Child Named Attribute"),
+				LOCTEXT("AddChildNamedAttribute_Tooltip", "Add a new child named attribute"),
 				FSlateIcon(FAppStyle::GetAppStyleSetName(), "Icons.Plus"),
 				AddNewAttributeAction);
 
