Skeleton Templates - Expand tree item on adding child attribute

diff --git a/Engine/Plugins/Experimental/Animation/SkeletonTemplateFramework/Source/Editor/Private/SkeletonTemplateFramework/SkeletonTemplateNamedAttributesEditor.cpp b/Engine/Plugins/Experimental/Animation/SkeletonTemplateFramework/Source/Editor/Private/SkeletonTemplateFramework/SkeletonTemplateNamedAttributesEditor.cpp
--- a/Engine/Plugins/Experimental/Animation/SkeletonTemplateFramework/Source/Editor/Private/SkeletonTemplateFramework/SkeletonTemplateNamedAttributesEditor.cpp
+++ b/Engine/Plugins/Experimental/Animation/SkeletonTemplateFramework/Source/Editor/Private/SkeletonTemplateFramework/SkeletonTemplateNamedAttributesEditor.cpp
@@ -427,6 +427,10 @@ namespace UE::Anim::STF
 				TreeView->RequestScrollIntoView(Item);
 				DeferredRenameRequest = Item;
 			}
+			else if (Item->Name == NewAttribute.ParentName)
+			{
+				TreeView->SetItemExpansion(Item, true);
+			}
 		}
 	}
 
