Skeleton Templates - Fix warning from missing hierarchy table plugin dependency

diff --git a/Engine/Plugins/Experimental/Animation/SkeletonTemplateFramework/SkeletonTemplateFramework.uplugin b/Engine/Plugins/Experimental/Animation/SkeletonTemplateFramework/SkeletonTemplateFramework.uplugin
--- a/Engine/Plugins/Experimental/Animation/SkeletonTemplateFramework/SkeletonTemplateFramework.uplugin
+++ b/Engine/Plugins/Experimental/Animation/SkeletonTemplateFramework/SkeletonTemplateFramework.uplugin
@@ -28,5 +28,9 @@
 		}
 	],
 	"Plugins": [
+		{
+			"Name": "HierarchyTable",
+			"Enabled": true
+		}
 	]
 }
