Skeleton Templates - Update default named attribute name to avoid name clashes

diff --git a/Engine/Plugins/Experimental/Animation/SkeletonTemplateFramework/Source/Editor/Private/SkeletonTemplateFramework/SkeletonTemplateNamedAttributesEditor.cpp b/Engine/Plugins/Experimental/Animation/SkeletonTemplateFramework/Source/Editor/Private/SkeletonTemplateFramework/SkeletonTemplateNamedAttributesEditor.cpp
--- a/Engine/Plugins/Experimental/Animation/SkeletonTemplateFramework/Source/Editor/Private/SkeletonTemplateFramework/SkeletonTemplateNamedAttributesEditor.cpp
+++ b/Engine/Plugins/Experimental/Animation/SkeletonTemplateFramework/Source/Editor/Private/SkeletonTemplateFramework/SkeletonTemplateNamedAttributesEditor.cpp
@@ -406,6 +406,18 @@ namespace UE::Anim::STF
 		FSkeletonNamedAttribute NewAttribute;
 		NewAttribute.Name = FName(LOCTEXT("NewAttribute", "NewAttribute").ToString());
 
+		if (SkeletonTemplate->FindNamedAttribute(NewAttribute.Name) != nullptr)
+		{
+			// Resolve name clash
+			int32 Suffix = 1;
+			do
+			{
+				NewAttribute.Name = FName(FText::Format(LOCTEXT("NewAttributeFormat", "NewAttribute_{0}"), Suffix).ToString());
+				++Suffix;
+			}
+			while (SkeletonTemplate->FindNamedAttribute(NewAttribute.Name) != nullptr);
+		}
+
 		TArray<TSharedPtr<SAttributesTreeView::FTreeItem>> SelectedItems;
 		TreeView->GetSelectedItems(SelectedItems);
 	
diff --git a/Engine/Plugins/Experimental/Animation/SkeletonTemplateFramework/Source/Runtime/Private/SkeletonTemplateFramework/SkeletonTemplate.cpp b/Engine/Plugins/Experimental/Animation/SkeletonTemplateFramework/Source/Runtime/Private/SkeletonTemplateFramework/SkeletonTemplate.cpp
--- a/Engine/Plugins/Experimental/Animation/SkeletonTemplateFramework/Source/Runtime/Private/SkeletonTemplateFramework/SkeletonTemplate.cpp
+++ b/Engine/Plugins/Experimental/Animation/SkeletonTemplateFramework/Source/Runtime/Private/SkeletonTemplateFramework/SkeletonTemplate.cpp
@@ -71,13 +71,13 @@ bool USkeletonTemplate::ReparentNamedAttribute(const FName AttributeName, const
 		return false; // Cannot set parent of self
 	}
 
-	FSkeletonNamedAttribute* AttributeToReparent = FindNamedAttribute(AttributeName);
+	FSkeletonNamedAttribute* AttributeToReparent = FindMutableNamedAttribute(AttributeName);
 	if (!AttributeToReparent)
 	{
 		return false; // Attribute does not exist
 	}
 
-	FSkeletonNamedAttribute* NewParentAttribute = FindNamedAttribute(NewParentName);
+	FSkeletonNamedAttribute* NewParentAttribute = FindMutableNamedAttribute(NewParentName);
 	if (!NewParentAttribute)
 	{
 		return false; // Parent does not exist
@@ -93,7 +93,7 @@ bool USkeletonTemplate::ReparentNamedAttribute(const FName AttributeName, const
 				AncestorAttribute->ParentName = AttributeToReparent->ParentName;
 				break;
 			}
-			AncestorAttribute = FindNamedAttribute(AncestorAttribute->ParentName);
+			AncestorAttribute = FindMutableNamedAttribute(AncestorAttribute->ParentName);
 		}
 	}
 
@@ -145,7 +145,7 @@ bool USkeletonTemplate::RemoveNamedAttribute(const FName AttributeToRemoveName,
 
 bool USkeletonTemplate::SetNamedAttributeType(const FName AttributeName, const TObjectPtr<const UScriptStruct> NewType)
 {
-	FSkeletonNamedAttribute* NamedAttribute = FindNamedAttribute(AttributeName);
+	FSkeletonNamedAttribute* NamedAttribute = FindMutableNamedAttribute(AttributeName);
 
 	if (NamedAttribute)
 	{
@@ -156,7 +156,15 @@ bool USkeletonTemplate::SetNamedAttributeType(const FName AttributeName, const T
 	return false;
 }
 
-FSkeletonNamedAttribute* USkeletonTemplate::FindNamedAttribute(const FName AttributeName)
+const FSkeletonNamedAttribute* const USkeletonTemplate::FindNamedAttribute(const FName AttributeName) const
+{
+	return NamedAttributes.FindByPredicate([AttributeName](const FSkeletonNamedAttribute& Candidate)
+			{
+				return Candidate.Name == AttributeName;
+			});
+}
+
+FSkeletonNamedAttribute* USkeletonTemplate::FindMutableNamedAttribute(const FName AttributeName)
 {
 	return NamedAttributes.FindByPredicate([AttributeName](const FSkeletonNamedAttribute& Candidate)
 		{
diff --git a/Engine/Plugins/Experimental/Animation/SkeletonTemplateFramework/Source/Runtime/Public/SkeletonTemplateFramework/SkeletonTemplate.h b/Engine/Plugins/Experimental/Animation/SkeletonTemplateFramework/Source/Runtime/Public/SkeletonTemplateFramework/SkeletonTemplate.h
--- a/Engine/Plugins/Experimental/Animation/SkeletonTemplateFramework/Source/Runtime/Public/SkeletonTemplateFramework/SkeletonTemplate.h
+++ b/Engine/Plugins/Experimental/Animation/SkeletonTemplateFramework/Source/Runtime/Public/SkeletonTemplateFramework/SkeletonTemplate.h
@@ -94,8 +94,10 @@ public:
 
 	SKELETONTEMPLATEFRAMEWORKRUNTIME_API bool SetNamedAttributeType(const FName AttributeName, const TObjectPtr<const UScriptStruct> NewType);
 
+	SKELETONTEMPLATEFRAMEWORKRUNTIME_API const FSkeletonNamedAttribute* const FindNamedAttribute(const FName AttributeName) const;
+
 protected:
-	FSkeletonNamedAttribute* FindNamedAttribute(const FName AttributeName);
+	FSkeletonNamedAttribute* FindMutableNamedAttribute(const FName AttributeName);
 
 	int32 FindNamedAttributeIndex(const FName AttributeName);
 
