Skeleton Templates - Named attribute sets editor
* Excluded sub-parts for now

diff --git a/Engine/Plugins/Experimental/Animation/SkeletonTemplateFramework/Source/Editor/Private/SkeletonTemplateFramework/SkeletonTemplateEditorToolkit.cpp b/Engine/Plugins/Experimental/Animation/SkeletonTemplateFramework/Source/Editor/Private/SkeletonTemplateFramework/SkeletonTemplateEditorToolkit.cpp
--- a/Engine/Plugins/Experimental/Animation/SkeletonTemplateFramework/Source/Editor/Private/SkeletonTemplateFramework/SkeletonTemplateEditorToolkit.cpp
+++ b/Engine/Plugins/Experimental/Animation/SkeletonTemplateFramework/Source/Editor/Private/SkeletonTemplateFramework/SkeletonTemplateEditorToolkit.cpp
@@ -6,6 +6,7 @@
 #include "PropertyEditorModule.h"
 #include "SkeletonTemplateFramework/SkeletonTemplate.h"
 #include "SkeletonTemplateFramework/SkeletonTemplateNamedAttributesEditor.h"
+#include "SkeletonTemplateFramework/SkeletonTemplateNamedAttributeSetsEditor.h"
 #include "Widgets/Docking/SDockTab.h"
 
 #define LOCTEXT_NAMESPACE "SkeletonTemplateEditorToolkit"
@@ -15,12 +16,14 @@ namespace UE::Anim::STF
 	struct FEditorTabs
 	{
 		static const FName AttributesId;
+		static const FName AttributeSetsId;
 		static const FName TableViewId;
 		static const FName DetailsId;
 	};
 }
 
 const FName UE::Anim::STF::FEditorTabs::AttributesId(TEXT("Attributes"));
+const FName UE::Anim::STF::FEditorTabs::AttributeSetsId(TEXT("AttributeSets"));
 const FName UE::Anim::STF::FEditorTabs::TableViewId(TEXT("TableView"));
 const FName UE::Anim::STF::FEditorTabs::DetailsId(TEXT("Details"));
 
@@ -37,9 +40,21 @@ void FSkeletonTemplateEditorToolkit::InitEditor(const TArray<UObject*>& InObject
 			FTabManager::NewPrimaryArea()->SetOrientation(Orient_Horizontal)
 			->Split
 			(
-				FTabManager::NewStack()
+				FTabManager::NewSplitter()
+				->SetOrientation(Orient_Vertical)
 				->SetSizeCoefficient(0.25f)
-				->AddTab(FEditorTabs::AttributesId, ETabState::OpenedTab)
+				->Split
+				(
+					FTabManager::NewStack()
+					->SetSizeCoefficient(0.5f)
+					->AddTab(FEditorTabs::AttributesId, ETabState::OpenedTab)
+				)
+				->Split
+				(
+					FTabManager::NewStack()
+					->SetSizeCoefficient(0.5f)
+					->AddTab(FEditorTabs::AttributeSetsId, ETabState::OpenedTab)
+				)
 			)
 			->Split
 			(
@@ -72,6 +87,12 @@ void FSkeletonTemplateEditorToolkit::RegisterTabSpawners(const TSharedRef<class
 		.SetTooltipText(LOCTEXT("AttributesTabMenu_ToolTip", "Shows the attributes panel"))
 		.SetGroup(WorkspaceMenuCategoryRef)
 		.SetIcon(FSlateIcon(FAppStyle::GetAppStyleSetName(), "LevelEditor.Tabs.Outliner"));
+
+	InTabManager->RegisterTabSpawner(FEditorTabs::AttributeSetsId, FOnSpawnTab::CreateSP(this, &FSkeletonTemplateEditorToolkit::SpawnTabAttributeSets))
+		.SetDisplayName(LOCTEXT("AttributeSetsTabMenu_Description", "Attribute Sets"))
+		.SetTooltipText(LOCTEXT("AttributeSetsTabMenu_ToolTip", "Shows the attribute sets panel"))
+		.SetGroup(WorkspaceMenuCategoryRef)
+		.SetIcon(FSlateIcon(FAppStyle::GetAppStyleSetName(), "LevelEditor.Tabs.Outliner"));
 	
 	InTabManager->RegisterTabSpawner(FEditorTabs::TableViewId, FOnSpawnTab::CreateSP(this, &FSkeletonTemplateEditorToolkit::SpawnTabTableView))
 		.SetDisplayName(LOCTEXT("TableViewTabMenu_Description", "Table View"))
@@ -124,6 +145,14 @@ TSharedRef<SDockTab> FSkeletonTemplateEditorToolkit::SpawnTabAttributes(const FS
 		];
 }
 
+TSharedRef<SDockTab> FSkeletonTemplateEditorToolkit::SpawnTabAttributeSets(const FSpawnTabArgs& Args)
+{
+	return SNew(SDockTab)
+		[
+			SAssignNew(AttributeSetsTreeView, UE::Anim::STF::SAttributeSetsTreeView, SkeletonTemplate)
+		];
+}
+
 TSharedRef<SDockTab> FSkeletonTemplateEditorToolkit::SpawnTabTableView(const FSpawnTabArgs& Args)
 {
 	return SNew(SDockTab);
diff --git a/Engine/Plugins/Experimental/Animation/SkeletonTemplateFramework/Source/Editor/Private/SkeletonTemplateFramework/SkeletonTemplateEditorToolkit.h b/Engine/Plugins/Experimental/Animation/SkeletonTemplateFramework/Source/Editor/Private/SkeletonTemplateFramework/SkeletonTemplateEditorToolkit.h
--- a/Engine/Plugins/Experimental/Animation/SkeletonTemplateFramework/Source/Editor/Private/SkeletonTemplateFramework/SkeletonTemplateEditorToolkit.h
+++ b/Engine/Plugins/Experimental/Animation/SkeletonTemplateFramework/Source/Editor/Private/SkeletonTemplateFramework/SkeletonTemplateEditorToolkit.h
@@ -10,6 +10,7 @@ class SDockTab;
 namespace UE::Anim::STF
 {
 	class SAttributesTreeView;
+	class SAttributeSetsTreeView;
 }
 
 class FSkeletonTemplateEditorToolkit : public FAssetEditorToolkit
@@ -27,6 +28,7 @@ public:
 
 private:
 	TSharedRef<SDockTab> SpawnTabAttributes(const FSpawnTabArgs& Args);
+	TSharedRef<SDockTab> SpawnTabAttributeSets(const FSpawnTabArgs& Args);
 	TSharedRef<SDockTab> SpawnTabTableView(const FSpawnTabArgs& Args);
 	TSharedRef<SDockTab> SpawnTabDetails(const FSpawnTabArgs& Args);
 
@@ -35,4 +37,6 @@ private:
 	TSharedPtr<IDetailsView> DetailsView;
 
 	TSharedPtr<UE::Anim::STF::SAttributesTreeView> AttributesTreeView;
+
+	TSharedPtr<UE::Anim::STF::SAttributeSetsTreeView> AttributeSetsTreeView;
 };
\ No newline at end of file
diff --git a/Engine/Plugins/Experimental/Animation/SkeletonTemplateFramework/Source/Editor/Private/SkeletonTemplateFramework/SkeletonTemplateNamedAttributeSetsEditor.cpp b/Engine/Plugins/Experimental/Animation/SkeletonTemplateFramework/Source/Editor/Private/SkeletonTemplateFramework/SkeletonTemplateNamedAttributeSetsEditor.cpp
new file mode 100644
--- /dev/null
+++ b/Engine/Plugins/Experimental/Animation/SkeletonTemplateFramework/Source/Editor/Private/SkeletonTemplateFramework/SkeletonTemplateNamedAttributeSetsEditor.cpp
@@ -0,0 +1,383 @@
+// Copyright Epic Games, Inc. All Rights Reserved.
+
+#include "SkeletonTemplateNamedAttributeSetsEditor.h"
+
+#include "DragAndDrop/DecoratedDragDropOp.h"
+#include "Framework/MultiBox/MultiBoxBuilder.h"
+#include "Modules/ModuleManager.h"
+#include "SkeletonTemplateFramework/SkeletonTemplate.h"
+#include "SkeletonTemplateFramework/SkeletonTemplateNamedAttributesEditor.h"
+#include "StructViewerFilter.h"
+#include "StructViewerModule.h"
+#include "Widgets/Input/SComboButton.h"
+#include "Widgets/Text/SInlineEditableTextBlock.h"
+#include "Widgets/Views/STreeView.h"
+
+#define LOCTEXT_NAMESPACE "SkeletonTemplateNamedAttributeSetsEditor"
+
+namespace UE::Anim::STF
+{
+	struct FNamedAttributeSetsColumns
+	{
+		static const FName NameId;
+		static const FName TypeId;
+	};
+
+	const FName FNamedAttributeSetsColumns::NameId(TEXT("Name"));
+	const FName FNamedAttributeSetsColumns::TypeId(TEXT("Type"));
+
+	class SNamedElementSetTableRow : public SMultiColumnTableRow<TSharedPtr<SAttributeSetsTreeView::FTreeItem>>
+	{
+	public:
+		DECLARE_DELEGATE_RetVal_TwoParams(bool /* Success */, FOnRenamed, FName /* Old Name */, FName /* New Name */);
+		DECLARE_DELEGATE_TwoParams(FOnAddAttributeToSet, FName /* Attribute Set Name */, FName /* Attribute Name */);
+
+		SLATE_BEGIN_ARGS(SNamedElementSetTableRow) { }
+			SLATE_ARGUMENT(TSharedPtr<SAttributeSetsTreeView::FTreeItem>, TreeItem)
+			SLATE_EVENT(FOnRenamed, OnRenamed);
+			SLATE_EVENT(FOnAddAttributeToSet, OnAddAttributeToSet)
+		SLATE_END_ARGS()
+
+		TOptional<EItemDropZone> OnCanAcceptDrop(const FDragDropEvent& DragDropEvent, EItemDropZone DropZone, TSharedPtr<SAttributeSetsTreeView::FTreeItem> TargetItem)
+		{
+			using namespace UE::Anim::STF;
+
+			TOptional<EItemDropZone> ReturnedDropZone;
+	
+			const TSharedPtr<FNamedElementDragDropOp> DragDropOp = DragDropEvent.GetOperationAs<FNamedElementDragDropOp>();
+			if (DragDropOp.IsValid() && TreeItem->ItemType == SAttributeSetsTreeView::FTreeItem::EType::AttributeSet)
+			{
+        		ReturnedDropZone = EItemDropZone::BelowItem;
+			}
+	
+			return ReturnedDropZone;
+		}
+
+		FReply OnAcceptDrop(const FDragDropEvent& DragDropEvent, EItemDropZone DropZone, TSharedPtr<SAttributeSetsTreeView::FTreeItem> TargetItem)
+		{
+			const TSharedPtr<FNamedElementDragDropOp> DragDropOp = DragDropEvent.GetOperationAs<FNamedElementDragDropOp>();
+			if (!DragDropOp.IsValid())
+			{
+				return FReply::Unhandled();
+			}
+
+			OnAddAttributeToSet.ExecuteIfBound(TreeItem->Name, DragDropOp->Element.Pin()->Name);
+
+			return FReply::Handled();
+		}
+
+		void OnCommitRename(const FText& InText, ETextCommit::Type CommitInfo)
+		{
+			const FName OldName = TreeItem->Name;
+			const FName NewName = FName(InText.ToString());
+
+			if (OnRenamed.IsBound())
+			{
+				const bool bSuccess = OnRenamed.Execute(OldName, NewName);
+				if (bSuccess)
+				{
+					TreeItem->Name = NewName;
+				}
+			}
+		}
+
+		void Construct(const FArguments& InArgs, const TSharedRef<STableViewBase>& InOwnerTableView)
+		{
+			TreeItem = InArgs._TreeItem;
+			OnRenamed = InArgs._OnRenamed;
+			OnAddAttributeToSet = InArgs._OnAddAttributeToSet;
+
+			FSuperRowType::Construct(
+				FSuperRowType::FArguments()
+					.OnCanAcceptDrop(this, &SNamedElementSetTableRow::OnCanAcceptDrop)
+					.OnAcceptDrop(this, &SNamedElementSetTableRow::OnAcceptDrop),
+				InOwnerTableView);
+		}
+
+		TSharedRef<SWidget> GenerateWidgetForColumn(const FName& ColumnName) override
+		{
+			if (ColumnName == UE::Anim::STF::FNamedAttributeSetsColumns::NameId)
+			{
+				TSharedPtr<SInlineEditableTextBlock> InlineWidget;
+
+				TSharedRef<SHorizontalBox> HorizontalBox = SNew(SHorizontalBox)
+					+ SHorizontalBox::Slot()
+						.AutoWidth()
+						[
+							SNew(SExpanderArrow, SharedThis(this))
+								.ShouldDrawWires(true)
+						]
+					+ SHorizontalBox::Slot()
+						.AutoWidth()
+						[
+							SAssignNew(InlineWidget, SInlineEditableTextBlock)
+								.Text_Lambda([this]()
+								{
+									return FText::FromName(TreeItem->Name);
+								})
+								.OnTextCommitted(this, &SNamedElementSetTableRow::OnCommitRename)
+						];
+
+				TreeItem->OnRenameRequested.BindSP(InlineWidget.Get(), &SInlineEditableTextBlock::EnterEditingMode);
+
+				return HorizontalBox;
+			}
+			else if (ColumnName == UE::Anim::STF::FNamedAttributeSetsColumns::TypeId)
+			{
+				if (TreeItem->ItemType == SAttributeSetsTreeView::FTreeItem::EType::AttributeSet)
+				{
+					return SNullWidget::NullWidget;
+				}
+				else
+				{
+					return SNew(STextBlock)
+						.Text_Lambda([this]()
+							{
+								return FText::FromName(TreeItem->Type);
+							});
+				}
+			}
+			else
+			{
+				return SNullWidget::NullWidget;
+			}
+		}
+
+	private:
+		TSharedPtr<SAttributeSetsTreeView::FTreeItem> TreeItem;
+		FOnAddAttributeToSet OnAddAttributeToSet;
+		FOnRenamed OnRenamed;
+	};
+
+	void SAttributeSetsTreeView::Construct(const FArguments& InArgs, TObjectPtr<USkeletonTemplate> InSkeletonTemplate)
+	{
+		SkeletonTemplate = InSkeletonTemplate;
+
+		ChildSlot
+		[
+			SAssignNew(TreeView, STreeView<TSharedPtr<FTreeItem>>)
+				.TreeItemsSource(&RootItems)
+				.OnGenerateRow(this, &SAttributeSetsTreeView::TreeView_GenerateItemRow)
+				.OnGetChildren(this, &SAttributeSetsTreeView::TreeView_HandleGetChildren)
+				.OnContextMenuOpening(this, &SAttributeSetsTreeView::TreeView_HandleContextMenuOpening)
+				.OnItemScrolledIntoView(this, &SAttributeSetsTreeView::TreeView_OnItemScrolledIntoView)
+				.HighlightParentNodesForSelection(true)
+				.HeaderRow
+				(
+					SNew(SHeaderRow)
+						+SHeaderRow::Column(UE::Anim::STF::FNamedAttributeSetsColumns::NameId)
+							.FillWidth(0.5f)
+							.DefaultLabel(LOCTEXT("NameLabel", "Name"))
+						+SHeaderRow::Column(UE::Anim::STF::FNamedAttributeSetsColumns::TypeId)
+							.FillWidth(0.5f)
+							.DefaultLabel(LOCTEXT("TypeLabel", "Type"))
+				)
+		];
+
+		RegenerateTreeViewItems();
+
+		// Expand all tree items on construction
+		for (TSharedPtr<FTreeItem> TreeItem : GetAllTreeItems())
+		{
+			TreeView->SetItemExpansion(TreeItem, true);
+		}
+	}
+
+	void SAttributeSetsTreeView::TreeView_HandleGetChildren(TSharedPtr<FTreeItem> InItem, TArray<TSharedPtr<FTreeItem>>& OutChildren) const
+	{
+		OutChildren.Append(InItem->Children);
+	}
+
+	void SAttributeSetsTreeView::TreeView_OnItemScrolledIntoView(TSharedPtr<FTreeItem> InItem, const TSharedPtr<ITableRow>& InWidget)
+	{
+		if (DeferredRenameRequest.IsValid())
+		{
+			DeferredRenameRequest->OnRenameRequested.ExecuteIfBound();
+			DeferredRenameRequest.Reset();
+		}
+	}
+
+	TSharedRef<ITableRow> SAttributeSetsTreeView::TreeView_GenerateItemRow(TSharedPtr<SAttributeSetsTreeView::FTreeItem> Item, const TSharedRef<STableViewBase>& OwnerTable)
+	{
+		if (Item->ItemType == FTreeItem::EType::AttributeSet)
+		{
+			return SNew(SNamedElementSetTableRow, OwnerTable)
+				.TreeItem(Item)
+				.OnAddAttributeToSet_Lambda([this, Item](const FName AttributeSet, const FName AttributeName)
+					{
+						SkeletonTemplate->AddAttributeToSet(AttributeSet, AttributeName);
+						TreeView->SetItemExpansion(Item, true);
+						RegenerateTreeViewItems();
+					})
+				.OnRenamed_Lambda([this](const FName OldName, const FName NewName) -> bool
+				{
+					return SkeletonTemplate->RenameNamedAttributeSet(OldName, NewName);
+				});
+		}
+		else
+		{
+			return SNew(SNamedElementSetTableRow, OwnerTable)
+				.TreeItem(Item);
+		}
+	}
+
+	TSharedPtr<SWidget> SAttributeSetsTreeView::TreeView_HandleContextMenuOpening()
+	{
+		FMenuBuilder MenuBuilder(true, nullptr);
+
+		TArray<TSharedPtr<FTreeItem>> SelectedItems;
+		TreeView->GetSelectedItems(SelectedItems);
+
+		if (!SelectedItems.IsEmpty())
+		{
+			TSharedPtr<FTreeItem> SelectedItem = SelectedItems[0];
+
+			switch (SelectedItem->ItemType)
+			{
+				case FTreeItem::EType::AttributeSet:
+				{
+					MenuBuilder.AddMenuEntry(
+						LOCTEXT("RenameNamedAttributeSet", "Rename Named Attribute Set"),
+						LOCTEXT("RenameNamedAttributeSet_Tooltip", "Renames the selected named attribute set"),
+						FSlateIcon(FAppStyle::GetAppStyleSetName(), "GenericCommands.Rename"),
+						FUIAction(
+							FExecuteAction::CreateLambda([this, SelectedItem]()
+							{
+								SelectedItem->OnRenameRequested.ExecuteIfBound();
+							})
+						));
+
+					MenuBuilder.AddMenuEntry(
+						LOCTEXT("DeleteNamedAttributeSet", "Delete Named Attribute Set"),
+						LOCTEXT("DeleteNamedAttributeSet_Tooltip", "Deletes the selected named attribute set"),
+						FSlateIcon(FAppStyle::GetAppStyleSetName(), "Icons.Delete"),
+						FUIAction(
+							FExecuteAction::CreateLambda([this, SelectedItem]()
+							{
+								SkeletonTemplate->RemoveAttributeSet(SelectedItem->Name);
+								RegenerateTreeViewItems();
+							})
+						));
+
+					break;
+				}
+				case FTreeItem::EType::Attribute:
+				{
+					MenuBuilder.AddMenuEntry(
+						LOCTEXT("AddNamedAttribute", "Remove Named Attribute"),
+						LOCTEXT("AddNamedAttribute_Tooltip", "Remove the selected named attribute from the set"),
+						FSlateIcon(FAppStyle::GetAppStyleSetName(), "Icons.Minus"),
+						FUIAction(
+							FExecuteAction::CreateLambda([this, SelectedItem]()
+							{
+								SkeletonTemplate->RemoveAttributeFromSet(SelectedItem->Parent->Name, SelectedItem->Name);
+								RegenerateTreeViewItems();
+							})
+						));
+
+					break;
+				}
+				case FTreeItem::EType::AttributePart:
+				{
+				}
+			}
+		}
+		else
+		{
+			MenuBuilder.AddMenuEntry(
+				LOCTEXT("AddNamedAttributeSet", "Add Named Attribute Set"),
+				LOCTEXT("AddNamedAttributeSet_Tooltip", "Add a new named attribute set"),
+				FSlateIcon(FAppStyle::GetAppStyleSetName(), "Icons.Plus"),
+				FUIAction(
+					FExecuteAction::CreateLambda([this]()
+					{
+						SkeletonTemplate->AddAttributeSet("NewAttributeSet");
+						RegenerateTreeViewItems();
+
+						for (TSharedPtr<FTreeItem> Item : GetAllTreeItems())
+						{
+							if (Item->ItemType == FTreeItem::EType::AttributeSet && Item->Name == "NewAttributeSet")
+							{
+								TreeView->RequestScrollIntoView(Item);
+								DeferredRenameRequest = Item;
+							}
+						}
+					})
+				));
+		}
+
+		return MenuBuilder.MakeWidget();
+	}
+
+	TArray<TSharedPtr<UE::Anim::STF::SAttributeSetsTreeView::FTreeItem>> SAttributeSetsTreeView::GetAllTreeItems()
+	{
+		using namespace UE::Anim::STF;
+
+		TArray<TSharedPtr<SAttributeSetsTreeView::FTreeItem>> AllItems;
+		AllItems.Append(RootItems);
+	
+		for (int32 Index = 0; Index < AllItems.Num(); ++Index)
+		{
+			AllItems.Append(AllItems[Index]->Children);
+		}
+
+		return AllItems;
+	}
+
+	void SAttributeSetsTreeView::RegenerateTreeViewItems()
+	{
+		// Make note of all tree items currently expanded
+		TSet<FName> ExpandedAttributeNames;
+		{
+			for (TSharedPtr<FTreeItem> TreeItem : GetAllTreeItems())
+			{
+				if (TreeView->IsItemExpanded(TreeItem))
+				{
+					ExpandedAttributeNames.Add(TreeItem->Name);
+				}
+			}
+		}
+
+		// Rebuild items
+		{
+			RootItems.Reset();
+	
+			for (const FSkeletonNamedAttributeSet& NamedAttributeSet : SkeletonTemplate->GetNamedAttributeSets())
+			{
+				TSharedPtr<FTreeItem> SetItem = MakeShared<FTreeItem>();
+				SetItem->Name = NamedAttributeSet.Name;
+				SetItem->ItemType = FTreeItem::EType::AttributeSet;
+				SetItem->Type = NAME_None;
+
+				for (const FName AttributeName : NamedAttributeSet.NamedAttributes)
+				{
+					TSharedPtr<FTreeItem> AttributeItem = MakeShared<FTreeItem>();
+					AttributeItem->Name = AttributeName;
+					AttributeItem->ItemType = FTreeItem::EType::Attribute;
+					AttributeItem->Type = SkeletonTemplate->FindNamedAttribute(AttributeName)->Type.GetFName();
+					AttributeItem->Parent = SetItem;
+
+					SetItem->Children.Add(AttributeItem);
+				}
+
+				RootItems.Add(SetItem);
+			}
+		}
+
+		// Update tree view and restore tree item expanded states
+		{
+			check(TreeView);
+			TreeView->RequestTreeRefresh();
+
+			for (TSharedPtr<FTreeItem> TreeItem : GetAllTreeItems())
+			{
+				if (ExpandedAttributeNames.Contains(TreeItem->Name))
+				{
+					TreeView->SetItemExpansion(TreeItem, true);
+				}
+			}
+		}
+	}
+}
+
+#undef LOCTEXT_NAMESPACE
\ No newline at end of file
diff --git a/Engine/Plugins/Experimental/Animation/SkeletonTemplateFramework/Source/Editor/Private/SkeletonTemplateFramework/SkeletonTemplateNamedAttributeSetsEditor.h b/Engine/Plugins/Experimental/Animation/SkeletonTemplateFramework/Source/Editor/Private/SkeletonTemplateFramework/SkeletonTemplateNamedAttributeSetsEditor.h
new file mode 100644
--- /dev/null
+++ b/Engine/Plugins/Experimental/Animation/SkeletonTemplateFramework/Source/Editor/Private/SkeletonTemplateFramework/SkeletonTemplateNamedAttributeSetsEditor.h
@@ -0,0 +1,65 @@
+// Copyright Epic Games, Inc. All Rights Reserved.
+
+#pragma once
+
+#include "CoreMinimal.h"
+#include "Widgets/SCompoundWidget.h"
+
+class ITableRow;
+class STableViewBase;
+template<typename ItemType> class STreeView;
+class USkeletonTemplate;
+
+namespace UE::Anim::STF
+{
+	class SAttributeSetsTreeView : public SCompoundWidget
+	{
+	public:
+		SLATE_BEGIN_ARGS(SAttributeSetsTreeView) { }
+		SLATE_END_ARGS()
+
+		struct FTreeItem
+		{
+			enum class EType
+			{
+				AttributeSet,
+				Attribute,
+				AttributePart,
+			};
+
+			EType ItemType;
+			FName Name;
+			FName Type;
+			TSharedPtr<FTreeItem> Parent;
+			TArray<TSharedPtr<FTreeItem>> Children;
+
+			DECLARE_DELEGATE(FOnRenameRequested);
+			FOnRenameRequested OnRenameRequested;
+		};
+
+		void Construct(const FArguments& InArgs, TObjectPtr<USkeletonTemplate> InSkeletonTemplate);
+
+	private:
+		void TreeView_HandleGetChildren(TSharedPtr<FTreeItem> InItem, TArray<TSharedPtr<FTreeItem>>& OutChildren) const;
+		
+		void TreeView_OnItemScrolledIntoView(TSharedPtr<FTreeItem> InItem, const TSharedPtr<ITableRow>& InWidget);
+
+		TSharedRef<ITableRow> TreeView_GenerateItemRow(TSharedPtr<FTreeItem> Item, const TSharedRef<STableViewBase>& OwnerTable);
+		
+		TSharedPtr<SWidget> TreeView_HandleContextMenuOpening();
+
+	private:
+		TArray<TSharedPtr<FTreeItem>> GetAllTreeItems();
+
+		void RegenerateTreeViewItems();
+
+	private:
+		TSharedPtr<STreeView<TSharedPtr<FTreeItem>>> TreeView;
+
+		TArray<TSharedPtr<FTreeItem>> RootItems;
+
+		TSharedPtr<FTreeItem> DeferredRenameRequest;
+
+		TObjectPtr<USkeletonTemplate> SkeletonTemplate;
+	};
+}
\ No newline at end of file
diff --git a/Engine/Plugins/Experimental/Animation/SkeletonTemplateFramework/Source/Editor/Private/SkeletonTemplateFramework/SkeletonTemplateNamedAttributesEditor.cpp b/Engine/Plugins/Experimental/Animation/SkeletonTemplateFramework/Source/Editor/Private/SkeletonTemplateFramework/SkeletonTemplateNamedAttributesEditor.cpp
--- a/Engine/Plugins/Experimental/Animation/SkeletonTemplateFramework/Source/Editor/Private/SkeletonTemplateFramework/SkeletonTemplateNamedAttributesEditor.cpp
+++ b/Engine/Plugins/Experimental/Animation/SkeletonTemplateFramework/Source/Editor/Private/SkeletonTemplateFramework/SkeletonTemplateNamedAttributesEditor.cpp
@@ -2,7 +2,6 @@
 
 #include "SkeletonTemplateNamedAttributesEditor.h"
 
-#include "DragAndDrop/DecoratedDragDropOp.h"
 #include "Framework/MultiBox/MultiBoxBuilder.h"
 #include "Modules/ModuleManager.h"
 #include "SkeletonTemplateFramework/SkeletonTemplate.h"
@@ -25,32 +24,24 @@ namespace UE::Anim::STF
 	const FName FOutlinerColumns::NameId(TEXT("Name"));
 	const FName FOutlinerColumns::TypeId(TEXT("Type"));
 
-	class FNamedElementDragDropOp : public FDecoratedDragDropOp
+	TSharedRef<FNamedElementDragDropOp> FNamedElementDragDropOp::New(TWeakPtr<SAttributesTreeView::FTreeItem> InElement)
 	{
-	public:
-		DRAG_DROP_OPERATOR_TYPE(FNamedElementDragDropOp, FDecoratedDragDropOp)
-
-		static TSharedRef<FNamedElementDragDropOp> New(TWeakPtr<SAttributesTreeView::FTreeItem> InElement)
-		{
-			TSharedRef<FNamedElementDragDropOp> Operation = MakeShared<FNamedElementDragDropOp>();
-			Operation->Element = InElement;
-			Operation->Construct();
-			return Operation;
-		}
-
-		TSharedPtr<SWidget> GetDefaultDecorator() const override
-		{
-			return SNew(SBorder)
-				.Visibility(EVisibility::Visible)
-				.BorderImage(FAppStyle::GetBrush("Menu.Background"))
-				[
-					SNew(STextBlock)
-					.Text(FText::FromName(Element.Pin()->Name))
-				];
-		}
+		TSharedRef<FNamedElementDragDropOp> Operation = MakeShared<FNamedElementDragDropOp>();
+		Operation->Element = InElement;
+		Operation->Construct();
+		return Operation;
+	}
 
-		TWeakPtr<SAttributesTreeView::FTreeItem> Element;
-	};
+	TSharedPtr<SWidget> FNamedElementDragDropOp::GetDefaultDecorator() const
+	{
+		return SNew(SBorder)
+			.Visibility(EVisibility::Visible)
+			.BorderImage(FAppStyle::GetBrush("Menu.Background"))
+			[
+				SNew(STextBlock)
+				.Text(FText::FromName(Element.Pin()->Name))
+			];
+	}
 
 	class SNamedElementTableRow : public SMultiColumnTableRow<TSharedPtr<SAttributesTreeView::FTreeItem>>
 	{
@@ -256,7 +247,7 @@ namespace UE::Anim::STF
 
 	void SAttributesTreeView::TreeView_OnItemScrolledIntoView(TSharedPtr<FTreeItem> InItem, const TSharedPtr<ITableRow>& InWidget)
 	{
-		if(DeferredRenameRequest.IsValid())
+		if (DeferredRenameRequest.IsValid())
 		{
 			DeferredRenameRequest->OnRenameRequested.ExecuteIfBound();
 			DeferredRenameRequest.Reset();
diff --git a/Engine/Plugins/Experimental/Animation/SkeletonTemplateFramework/Source/Editor/Private/SkeletonTemplateFramework/SkeletonTemplateNamedAttributesEditor.h b/Engine/Plugins/Experimental/Animation/SkeletonTemplateFramework/Source/Editor/Private/SkeletonTemplateFramework/SkeletonTemplateNamedAttributesEditor.h
--- a/Engine/Plugins/Experimental/Animation/SkeletonTemplateFramework/Source/Editor/Private/SkeletonTemplateFramework/SkeletonTemplateNamedAttributesEditor.h
+++ b/Engine/Plugins/Experimental/Animation/SkeletonTemplateFramework/Source/Editor/Private/SkeletonTemplateFramework/SkeletonTemplateNamedAttributesEditor.h
@@ -3,6 +3,7 @@
 #pragma once
 
 #include "CoreMinimal.h"
+#include "DragAndDrop/DecoratedDragDropOp.h"
 #include "Widgets/SCompoundWidget.h"
 
 class ITableRow;
@@ -59,4 +60,16 @@ namespace UE::Anim::STF
 
 		TObjectPtr<USkeletonTemplate> SkeletonTemplate;
 	};
+
+	class FNamedElementDragDropOp : public FDecoratedDragDropOp
+	{
+	public:
+		DRAG_DROP_OPERATOR_TYPE(FNamedElementDragDropOp, FDecoratedDragDropOp)
+
+		static TSharedRef<FNamedElementDragDropOp> New(TWeakPtr<SAttributesTreeView::FTreeItem> InElement);
+
+		TSharedPtr<SWidget> GetDefaultDecorator() const override;
+
+		TWeakPtr<SAttributesTreeView::FTreeItem> Element;
+	};
 }
\ No newline at end of file
diff --git a/Engine/Plugins/Experimental/Animation/SkeletonTemplateFramework/Source/Runtime/Private/SkeletonTemplateFramework/SkeletonTemplate.cpp b/Engine/Plugins/Experimental/Animation/SkeletonTemplateFramework/Source/Runtime/Private/SkeletonTemplateFramework/SkeletonTemplate.cpp
--- a/Engine/Plugins/Experimental/Animation/SkeletonTemplateFramework/Source/Runtime/Private/SkeletonTemplateFramework/SkeletonTemplate.cpp
+++ b/Engine/Plugins/Experimental/Animation/SkeletonTemplateFramework/Source/Runtime/Private/SkeletonTemplateFramework/SkeletonTemplate.cpp
@@ -156,6 +156,11 @@ bool USkeletonTemplate::SetNamedAttributeType(const FName AttributeName, const T
 	return false;
 }
 
+const TArray<FSkeletonNamedAttributeSet>& USkeletonTemplate::GetNamedAttributeSets() const
+{
+	return NamedAttributeSets;
+}
+
 const FSkeletonNamedAttribute* const USkeletonTemplate::FindNamedAttribute(const FName AttributeName) const
 {
 	return NamedAttributes.FindByPredicate([AttributeName](const FSkeletonNamedAttribute& Candidate)
@@ -164,6 +169,85 @@ const FSkeletonNamedAttribute* const USkeletonTemplate::FindNamedAttribute(const
 			});
 }
 
+bool USkeletonTemplate::AddAttributeSet(const FName AttributeSetName)
+{
+	FSkeletonNamedAttributeSet* AttributeSet = FindMutableNamedAttributeSet(AttributeSetName);
+	if (AttributeSet)
+	{
+		return false; // Set already exists with that name
+	}
+
+	FSkeletonNamedAttributeSet NewSet;
+	NewSet.Name = AttributeSetName;
+
+	NamedAttributeSets.Add(NewSet);
+	return true;
+}
+
+bool USkeletonTemplate::RemoveAttributeSet(const FName AttributeSetName)
+{
+	int32 AttributeSetIndex = FindNamedAttributeSetIndex(AttributeSetName);
+	if (AttributeSetIndex == INDEX_NONE)
+	{
+		return false; // Cannot find set with that name
+	}
+
+	NamedAttributeSets.RemoveAt(AttributeSetIndex);
+	return true;
+}
+
+bool USkeletonTemplate::AddAttributeToSet(const FName AttributeSetName, const FName AttributeName)
+{
+	FSkeletonNamedAttributeSet* AttributeSet = FindMutableNamedAttributeSet(AttributeSetName);
+	if (!AttributeSet)
+	{
+		return false;
+	}
+
+	const int32 Index = AttributeSet->NamedAttributes.AddUnique(AttributeName);
+
+	return Index != INDEX_NONE;
+}
+
+bool USkeletonTemplate::RemoveAttributeFromSet(const FName AttributeSetName, const FName AttributeName)
+{
+	FSkeletonNamedAttributeSet* AttributeSet = FindMutableNamedAttributeSet(AttributeSetName);
+	if (!AttributeSet)
+	{
+		return false;
+	}
+
+	const int32 RemoveCount = AttributeSet->NamedAttributes.Remove(AttributeName);
+
+	return RemoveCount > 0;
+}
+
+FSkeletonNamedAttributeSet* USkeletonTemplate::FindMutableNamedAttributeSet(const FName AttributeSetName)
+{
+	return NamedAttributeSets.FindByPredicate([AttributeSetName](const FSkeletonNamedAttributeSet& Candidate)
+		{
+			return Candidate.Name == AttributeSetName;
+		});
+}
+
+bool USkeletonTemplate::RenameNamedAttributeSet(const FName AttributeSetName, const FName NewAttributeSetName)
+{
+	FSkeletonNamedAttributeSet* AttributeSetToRename = FindMutableNamedAttributeSet(AttributeSetName);
+	if (!AttributeSetToRename)
+	{
+		return false; // No attribute set with old name
+	}
+
+	FSkeletonNamedAttributeSet* ExistingAttributeSet = FindMutableNamedAttributeSet(NewAttributeSetName);
+	if (ExistingAttributeSet)
+	{
+		return false; // Attribute set with new name already exists
+	}
+
+	AttributeSetToRename->Name = NewAttributeSetName;
+	return true;
+}
+
 FSkeletonNamedAttribute* USkeletonTemplate::FindMutableNamedAttribute(const FName AttributeName)
 {
 	return NamedAttributes.FindByPredicate([AttributeName](const FSkeletonNamedAttribute& Candidate)
@@ -172,6 +256,14 @@ FSkeletonNamedAttribute* USkeletonTemplate::FindMutableNamedAttribute(const FNam
 		});
 }
 
+int32 USkeletonTemplate::FindNamedAttributeSetIndex(const FName AttributeSetName)
+{
+	return NamedAttributeSets.IndexOfByPredicate([AttributeSetName](const FSkeletonNamedAttributeSet& Candidate)
+		{
+			return Candidate.Name == AttributeSetName;
+		});
+}
+
 int32 USkeletonTemplate::FindNamedAttributeIndex(const FName AttributeName)
 {
 	return NamedAttributes.IndexOfByPredicate([AttributeName](const FSkeletonNamedAttribute& Candidate)
diff --git a/Engine/Plugins/Experimental/Animation/SkeletonTemplateFramework/Source/Runtime/Public/SkeletonTemplateFramework/SkeletonTemplate.h b/Engine/Plugins/Experimental/Animation/SkeletonTemplateFramework/Source/Runtime/Public/SkeletonTemplateFramework/SkeletonTemplate.h
--- a/Engine/Plugins/Experimental/Animation/SkeletonTemplateFramework/Source/Runtime/Public/SkeletonTemplateFramework/SkeletonTemplate.h
+++ b/Engine/Plugins/Experimental/Animation/SkeletonTemplateFramework/Source/Runtime/Public/SkeletonTemplateFramework/SkeletonTemplate.h
@@ -94,11 +94,27 @@ public:
 
 	SKELETONTEMPLATEFRAMEWORKRUNTIME_API bool SetNamedAttributeType(const FName AttributeName, const TObjectPtr<const UScriptStruct> NewType);
 
+	SKELETONTEMPLATEFRAMEWORKRUNTIME_API const TArray<FSkeletonNamedAttributeSet>& GetNamedAttributeSets() const;
+
 	SKELETONTEMPLATEFRAMEWORKRUNTIME_API const FSkeletonNamedAttribute* const FindNamedAttribute(const FName AttributeName) const;
 
+	SKELETONTEMPLATEFRAMEWORKRUNTIME_API bool AddAttributeSet(const FName AttributeSetName);
+
+	SKELETONTEMPLATEFRAMEWORKRUNTIME_API bool RemoveAttributeSet(const FName AttributeSetName);
+	
+	SKELETONTEMPLATEFRAMEWORKRUNTIME_API bool AddAttributeToSet(const FName AttributeSetName, const FName AttributeName);
+
+	SKELETONTEMPLATEFRAMEWORKRUNTIME_API bool RemoveAttributeFromSet(const FName AttributeSetName, const FName AttributeName);
+
+	SKELETONTEMPLATEFRAMEWORKRUNTIME_API bool RenameNamedAttributeSet(const FName AttributeSetName, const FName NewAttributeSetName);
+
 protected:
 	FSkeletonNamedAttribute* FindMutableNamedAttribute(const FName AttributeName);
 
+	FSkeletonNamedAttributeSet* FindMutableNamedAttributeSet(const FName AttributeSetName);
+
+	int32 FindNamedAttributeSetIndex(const FName AttributeSetName);
+	
 	int32 FindNamedAttributeIndex(const FName AttributeName);
 
 	void SortNamedAttributes();
