Initial skeleton binding editor picker and layout

diff --git a/Engine/Plugins/Experimental/Animation/SkeletonTemplateFramework/Source/Editor/Private/SkeletonTemplateFramework/SkeletonBindingAssetDefinition.cpp b/Engine/Plugins/Experimental/Animation/SkeletonTemplateFramework/Source/Editor/Private/SkeletonTemplateFramework/SkeletonBindingAssetDefinition.cpp
--- a/Engine/Plugins/Experimental/Animation/SkeletonTemplateFramework/Source/Editor/Private/SkeletonTemplateFramework/SkeletonBindingAssetDefinition.cpp
+++ b/Engine/Plugins/Experimental/Animation/SkeletonTemplateFramework/Source/Editor/Private/SkeletonTemplateFramework/SkeletonBindingAssetDefinition.cpp
@@ -2,7 +2,7 @@
 
 #include "SkeletonTemplateFramework/SkeletonBindingAssetDefinition.h"
 #include "SkeletonTemplateFramework/SkeletonBinding.h"
-#include "Toolkits/SimpleAssetEditor.h"
+#include "SkeletonTemplateFramework/SkeletonBindingEditorToolkit.h"
 
 #define LOCTEXT_NAMESPACE "SkeletonTemplateFramework"
 
@@ -25,8 +25,8 @@ EAssetCommandResult UAssetDefinition_SkeletonBinding::OpenAssets(const FAssetOpe
 {
 	if (OpenArgs.OpenMethod == EAssetOpenMethod::Edit)
 	{
-		TArray<UObject*> Assets = OpenArgs.LoadObjects<UObject>();
-		FSimpleAssetEditor::CreateEditor(EToolkitMode::Standalone, OpenArgs.ToolkitHost, Assets);
+		const TArray<UObject*> Assets = OpenArgs.LoadObjects<UObject>();
+		MakeShared<FSkeletonBindingEditorToolkit>()->InitEditor(Assets);
 	}
 
 	return EAssetCommandResult::Handled;
diff --git a/Engine/Plugins/Experimental/Animation/SkeletonTemplateFramework/Source/Editor/Private/SkeletonTemplateFramework/SkeletonBindingEditorToolkit.cpp b/Engine/Plugins/Experimental/Animation/SkeletonTemplateFramework/Source/Editor/Private/SkeletonTemplateFramework/SkeletonBindingEditorToolkit.cpp
new file mode 100644
--- /dev/null
+++ b/Engine/Plugins/Experimental/Animation/SkeletonTemplateFramework/Source/Editor/Private/SkeletonTemplateFramework/SkeletonBindingEditorToolkit.cpp
@@ -0,0 +1,171 @@
+// Copyright Epic Games, Inc. All Rights Reserved.
+
+#include "SkeletonTemplateFramework/SkeletonBindingEditorToolkit.h"
+
+#include "Modules/ModuleManager.h"
+#include "PropertyEditorModule.h"
+#include "SkeletonTemplateFramework/SkeletonBinding.h"
+#include "Widgets/Docking/SDockTab.h"
+
+#define LOCTEXT_NAMESPACE "SkeletonBindingEditorToolkit"
+
+namespace UE::Anim::STF
+{
+	struct FSkeletonBindingEditorTabs
+	{
+		static const FName AttributesId;
+		static const FName AttributeSetsId;
+		static const FName AttributeMappingsId;
+		static const FName DetailsId;
+	};
+}
+
+const FName UE::Anim::STF::FSkeletonBindingEditorTabs::AttributesId(TEXT("Attributes"));
+const FName UE::Anim::STF::FSkeletonBindingEditorTabs::AttributeSetsId(TEXT("AttributeSets"));
+const FName UE::Anim::STF::FSkeletonBindingEditorTabs::AttributeMappingsId(TEXT("AttributeMappings"));
+const FName UE::Anim::STF::FSkeletonBindingEditorTabs::DetailsId(TEXT("Details"));
+
+void FSkeletonBindingEditorToolkit::InitEditor(const TArray<UObject*>& InObjects)
+{
+	using namespace UE::Anim::STF;
+
+	check(!InObjects.IsEmpty())
+	SkeletonBinding = CastChecked<USkeletonBinding>(InObjects[0]);
+
+	const TSharedRef<FTabManager::FLayout> Layout = FTabManager::NewLayout("SkeletonBindingEditorLayout")
+		->AddArea
+		(
+			FTabManager::NewPrimaryArea()->SetOrientation(Orient_Horizontal)
+			->Split
+			(
+				FTabManager::NewSplitter()
+				->SetOrientation(Orient_Vertical)
+				->SetSizeCoefficient(0.25f)
+				->Split
+				(
+					FTabManager::NewStack()
+					->SetSizeCoefficient(0.5f)
+					->AddTab(FSkeletonBindingEditorTabs::AttributesId, ETabState::OpenedTab)
+				)
+				->Split
+				(
+					FTabManager::NewStack()
+					->SetSizeCoefficient(0.5f)
+					->AddTab(FSkeletonBindingEditorTabs::AttributeSetsId, ETabState::OpenedTab)
+				)
+			)
+			->Split
+			(
+				FTabManager::NewStack()
+				->SetSizeCoefficient(0.5f)
+				->AddTab(FSkeletonBindingEditorTabs::AttributeMappingsId, ETabState::OpenedTab)
+			)
+			->Split
+			(
+				FTabManager::NewStack()
+				->SetSizeCoefficient(0.25f)
+				->AddTab(FSkeletonBindingEditorTabs::DetailsId, ETabState::OpenedTab)
+			)
+		);
+
+	FAssetEditorToolkit::InitAssetEditor(EToolkitMode::Standalone, {}, "SkeletonBindingEditor", Layout, true, true, InObjects);
+}
+
+void FSkeletonBindingEditorToolkit::RegisterTabSpawners(const TSharedRef<class FTabManager>& InTabManager)
+{
+	using namespace UE::Anim::STF;
+
+	WorkspaceMenuCategory = InTabManager->AddLocalWorkspaceMenuCategory(LOCTEXT("WorkspaceMenuCategory", "Skeleton Binding Editor"));
+	TSharedRef<FWorkspaceItem> WorkspaceMenuCategoryRef = WorkspaceMenuCategory.ToSharedRef();
+
+	FAssetEditorToolkit::RegisterTabSpawners(InTabManager);
+	
+	InTabManager->RegisterTabSpawner(FSkeletonBindingEditorTabs::AttributesId, FOnSpawnTab::CreateSP(this, &FSkeletonBindingEditorToolkit::SpawnTabAttributes))
+		.SetDisplayName(LOCTEXT("AttributesTabMenu_Description", "Attributes"))
+		.SetTooltipText(LOCTEXT("AttributesTabMenu_ToolTip", "Shows the attributes panel"))
+		.SetGroup(WorkspaceMenuCategoryRef)
+		.SetIcon(FSlateIcon(FAppStyle::GetAppStyleSetName(), "LevelEditor.Tabs.Outliner"));
+
+	InTabManager->RegisterTabSpawner(FSkeletonBindingEditorTabs::AttributeSetsId, FOnSpawnTab::CreateSP(this, &FSkeletonBindingEditorToolkit::SpawnTabAttributeSets))
+		.SetDisplayName(LOCTEXT("AttributeSetsTabMenu_Description", "Attribute Sets"))
+		.SetTooltipText(LOCTEXT("AttributeSetsTabMenu_ToolTip", "Shows the attribute sets panel"))
+		.SetGroup(WorkspaceMenuCategoryRef)
+		.SetIcon(FSlateIcon(FAppStyle::GetAppStyleSetName(), "LevelEditor.Tabs.Outliner"));
+	
+	InTabManager->RegisterTabSpawner(FSkeletonBindingEditorTabs::AttributeMappingsId, FOnSpawnTab::CreateSP(this, &FSkeletonBindingEditorToolkit::SpawnTabAttributeMappings))
+		.SetDisplayName(LOCTEXT("AttributeMappingsTabMenu_Description", "Attribute Mappings"))
+		.SetTooltipText(LOCTEXT("AttributeMappingsTabMenu_Description", "Shows the attribute mappings view panel"))
+		.SetGroup(WorkspaceMenuCategoryRef)
+		.SetIcon(FSlateIcon(FAppStyle::GetAppStyleSetName(), "LevelEditor.Tabs.Outliner"));
+
+	InTabManager->RegisterTabSpawner(FSkeletonBindingEditorTabs::DetailsId, FOnSpawnTab::CreateSP(this, &FSkeletonBindingEditorToolkit::SpawnTabDetails))
+		.SetDisplayName(LOCTEXT("DetailsTabMenu_Description", "Details"))
+		.SetTooltipText(LOCTEXT("DetailsTabMenu_Tooltip", "Shows the details panel"))
+		.SetGroup(WorkspaceMenuCategoryRef)
+		.SetIcon(FSlateIcon(FAppStyle::GetAppStyleSetName(), "LevelEditor.Tabs.Details"));
+}
+
+void FSkeletonBindingEditorToolkit::UnregisterTabSpawners(const TSharedRef<class FTabManager>& InTabManager)
+{
+	using namespace UE::Anim::STF;
+
+	FAssetEditorToolkit::UnregisterTabSpawners(InTabManager);
+
+	InTabManager->UnregisterTabSpawner(FSkeletonBindingEditorTabs::AttributesId);
+	InTabManager->UnregisterTabSpawner(FSkeletonBindingEditorTabs::AttributeSetsId);
+	InTabManager->UnregisterTabSpawner(FSkeletonBindingEditorTabs::AttributeMappingsId);
+	InTabManager->UnregisterTabSpawner(FSkeletonBindingEditorTabs::DetailsId);
+}
+
+FName FSkeletonBindingEditorToolkit::GetToolkitFName() const
+{
+	return "SkeletonBindingEditor";
+}
+
+FText FSkeletonBindingEditorToolkit::GetBaseToolkitName() const
+{
+	return LOCTEXT("ToolkitName", "Skeleton Binding Editor");
+}
+
+FString FSkeletonBindingEditorToolkit::GetWorldCentricTabPrefix() const
+{
+	return "SkeletonBindingEditor";
+}
+
+FLinearColor FSkeletonBindingEditorToolkit::GetWorldCentricTabColorScale() const
+{
+	return FLinearColor::White;
+}
+
+TSharedRef<SDockTab> FSkeletonBindingEditorToolkit::SpawnTabAttributes(const FSpawnTabArgs& Args)
+{
+	return SNew(SDockTab);
+}
+
+TSharedRef<SDockTab> FSkeletonBindingEditorToolkit::SpawnTabAttributeSets(const FSpawnTabArgs& Args)
+{
+	return SNew(SDockTab);
+}
+
+TSharedRef<SDockTab> FSkeletonBindingEditorToolkit::SpawnTabAttributeMappings(const FSpawnTabArgs& Args)
+{
+	return SNew(SDockTab);
+}
+
+TSharedRef<SDockTab> FSkeletonBindingEditorToolkit::SpawnTabDetails(const FSpawnTabArgs& Args)
+{
+	FPropertyEditorModule& PropertyEditorModule = FModuleManager::GetModuleChecked<FPropertyEditorModule>("PropertyEditor");
+	
+	FDetailsViewArgs DetailsViewArgs;
+	DetailsViewArgs.NameAreaSettings = FDetailsViewArgs::HideNameArea;
+	
+	DetailsView = PropertyEditorModule.CreateDetailView(DetailsViewArgs);
+	DetailsView->SetObjects(TArray<UObject*>{ SkeletonBinding });
+	
+	return SNew(SDockTab)
+		[
+			DetailsView.ToSharedRef()
+		];
+}
+
+#undef LOCTEXT_NAMESPACE
\ No newline at end of file
diff --git a/Engine/Plugins/Experimental/Animation/SkeletonTemplateFramework/Source/Editor/Private/SkeletonTemplateFramework/SkeletonBindingEditorToolkit.h b/Engine/Plugins/Experimental/Animation/SkeletonTemplateFramework/Source/Editor/Private/SkeletonTemplateFramework/SkeletonBindingEditorToolkit.h
new file mode 100644
--- /dev/null
+++ b/Engine/Plugins/Experimental/Animation/SkeletonTemplateFramework/Source/Editor/Private/SkeletonTemplateFramework/SkeletonBindingEditorToolkit.h
@@ -0,0 +1,39 @@
+// Copyright Epic Games, Inc. All Rights Reserved.
+
+#pragma once
+
+#include "CoreMinimal.h"
+#include "Toolkits/AssetEditorToolkit.h"
+
+class USkeletonBinding;
+class SDockTab;
+namespace UE::Anim::STF
+{
+	class SAttributesTreeView;
+	class SAttributeSetsTreeView;
+	class SAttributeMappingsTreeView;
+}
+
+class FSkeletonBindingEditorToolkit : public FAssetEditorToolkit
+{
+public:
+	void InitEditor(const TArray<UObject*>& InObjects);
+
+	// IToolkit
+	void RegisterTabSpawners(const TSharedRef<class FTabManager>& TabManager) override;
+	void UnregisterTabSpawners(const TSharedRef<class FTabManager>& TabManager) override;
+	FName GetToolkitFName() const override;
+	FText GetBaseToolkitName() const override;
+	FString GetWorldCentricTabPrefix() const override;
+	FLinearColor GetWorldCentricTabColorScale() const override;
+
+private:
+	TSharedRef<SDockTab> SpawnTabAttributes(const FSpawnTabArgs& Args);
+	TSharedRef<SDockTab> SpawnTabAttributeSets(const FSpawnTabArgs& Args);
+	TSharedRef<SDockTab> SpawnTabAttributeMappings(const FSpawnTabArgs& Args);
+	TSharedRef<SDockTab> SpawnTabDetails(const FSpawnTabArgs& Args);
+
+	TObjectPtr<USkeletonBinding> SkeletonBinding;
+
+	TSharedPtr<IDetailsView> DetailsView;
+};
\ No newline at end of file
diff --git a/Engine/Plugins/Experimental/Animation/SkeletonTemplateFramework/Source/Editor/Private/SkeletonTemplateFramework/SkeletonBindingFactory.cpp b/Engine/Plugins/Experimental/Animation/SkeletonTemplateFramework/Source/Editor/Private/SkeletonTemplateFramework/SkeletonBindingFactory.cpp
--- a/Engine/Plugins/Experimental/Animation/SkeletonTemplateFramework/Source/Editor/Private/SkeletonTemplateFramework/SkeletonBindingFactory.cpp
+++ b/Engine/Plugins/Experimental/Animation/SkeletonTemplateFramework/Source/Editor/Private/SkeletonTemplateFramework/SkeletonBindingFactory.cpp
@@ -1,8 +1,232 @@
 // Copyright Epic Games, Inc. All Rights Reserved.
 
 #include "SkeletonTemplateFramework/SkeletonBindingFactory.h"
+
+
+#include "Animation/Skeleton.h"
+#include "AssetRegistry/AssetData.h"
+#include "ContentBrowserModule.h"
+#include "Editor.h"
+#include "IContentBrowserSingleton.h"
+#include "Modules/ModuleManager.h"
 #include "SkeletonTemplateFramework/SkeletonBinding.h"
 #include "UObject/Package.h"
+#include "Widgets/Input/SButton.h"
+#include "Widgets/Layout/SBorder.h"
+#include "Widgets/Layout/SUniformGridPanel.h"
+#include "Widgets/SBoxPanel.h"
+#include "Widgets/SCompoundWidget.h"
+#include "Widgets/SWindow.h"
+
+#define LOCTEXT_NAMESPACE "SkeletonBindingFactory"
+
+class SSkeletonBindingCreateDialog final : public SCompoundWidget
+{
+public:
+	SLATE_BEGIN_ARGS(SSkeletonBindingCreateDialog){}
+	SLATE_END_ARGS()
+
+	void Construct(const FArguments& InArgs)
+	{
+		bOkClicked = false;
+
+		ChildSlot
+		[
+			SNew(SBorder)
+			.Visibility(EVisibility::Visible)
+			.BorderImage(FAppStyle::GetBrush("ChildWindow.Background"))
+			[
+				SNew(SVerticalBox)
+				+SVerticalBox::Slot()
+				.FillHeight(1.0f)
+				[
+					SNew(SVerticalBox)
+					+ SVerticalBox::Slot()
+						.AutoHeight()
+						.Padding(10.0f)
+						[
+							SNew(SVerticalBox)
+								+SVerticalBox::Slot()
+								.AutoHeight()
+								.Padding(0.0f, 0.0f, 0.0f, 5.0f)
+								[
+									CreateSkeletonPicker()
+								]
+						]
+					+ SVerticalBox::Slot()
+						.AutoHeight()
+						.Padding(10.0f, 0.0f, 10.0f, 0.0f)
+						[
+							SNew(SBorder)
+							.BorderImage(FAppStyle::GetBrush("NewAnimBlueprintDialog.AreaBorder"))
+							[
+								CreateSkeletonTemplatePicker()
+							]
+						]
+				]
+
+				// Create/Cancel buttons
+				+ SVerticalBox::Slot()
+					.AutoHeight()
+					.HAlign(HAlign_Right)
+					.VAlign(VAlign_Bottom)
+					.Padding(10.0f)
+					[
+						SNew(SUniformGridPanel)
+							.SlotPadding(FAppStyle::GetMargin("StandardDialog.SlotPadding"))
+							.MinDesiredSlotWidth(FAppStyle::GetFloat("StandardDialog.MinDesiredSlotWidth"))
+							.MinDesiredSlotHeight(FAppStyle::GetFloat("StandardDialog.MinDesiredSlotHeight"))
+							+ SUniformGridPanel::Slot(0,0)
+							[
+								SNew(SButton)
+									.ToolTipText(LOCTEXT("CreateButtonTooltip", "Create a new Skeleton Binding"))
+									.IsEnabled_Lambda([this]()
+										{
+											return SkeletonAsset.IsValid() && SkeletonTemplateAsset.IsValid();
+										})
+									.HAlign(HAlign_Center)
+									.ContentPadding(FAppStyle::GetMargin("StandardDialog.ContentPadding"))
+									.OnClicked(this, &SSkeletonBindingCreateDialog::OnCreateClicked)
+									.Text(LOCTEXT("CreateButtonText", "Create"))
+							]
+						+ SUniformGridPanel::Slot(1,0)
+						[
+							SNew(SButton)
+								.HAlign(HAlign_Center)
+								.ContentPadding( FAppStyle::GetMargin("StandardDialog.ContentPadding") )
+								.OnClicked(this, &SSkeletonBindingCreateDialog::OnCancelClicked)
+								.Text(LOCTEXT("CancelButtonText", "Cancel"))
+						]
+					]
+			]
+		];
+	}
+	
+	bool ConfigureProperties(TObjectPtr<USkeletonBindingFactory> InFactory)
+	{
+		Factory = InFactory;
+
+		const TSharedRef<SWindow> Window = SNew(SWindow)
+			.Title(LOCTEXT("CreateSkeletonBinding", "Create Skeleton Binding"))
+			.SizingRule(ESizingRule::UserSized)
+			.MinWidth(500.0f)
+			.MinHeight(800.0f)
+			.SupportsMinimize(false)
+			.SupportsMaximize(false)
+			[
+				AsShared()
+			];
+
+		PickerWindow = Window;
+
+		GEditor->EditorAddModalWindow(Window);
+		return bOkClicked;
+	}
+
+private:
+	TSharedRef<SWidget> CreateSkeletonTemplatePicker()
+	{
+		FContentBrowserModule& ContentBrowserModule = FModuleManager::Get().LoadModuleChecked<FContentBrowserModule>(TEXT("ContentBrowser"));
+
+		FAssetPickerConfig AssetPickerConfig;
+		AssetPickerConfig.Filter.ClassPaths.Add(USkeletonTemplate::StaticClass()->GetClassPathName());
+		AssetPickerConfig.OnAssetSelected = FOnAssetSelected::CreateLambda([this](const FAssetData& AssetData)
+			{
+				SkeletonTemplateAsset = AssetData;
+			});
+		AssetPickerConfig.bAllowNullSelection = false;
+		AssetPickerConfig.InitialAssetViewType = EAssetViewType::List;
+		AssetPickerConfig.InitialAssetSelection = SkeletonTemplateAsset;
+		AssetPickerConfig.bShowPathInColumnView = false;
+		AssetPickerConfig.bShowTypeInColumnView = false;
+		AssetPickerConfig.bFocusSearchBoxWhenOpened = false;
+	
+		return SNew(SVerticalBox)
+			+ SVerticalBox::Slot()
+				.FillHeight(1.0f)
+				.Padding(5.0f)
+				[
+					ContentBrowserModule.Get().CreateAssetPicker(AssetPickerConfig)
+				];
+	}
+
+	TSharedRef<SWidget> CreateSkeletonPicker()
+	{
+		FContentBrowserModule& ContentBrowserModule = FModuleManager::Get().LoadModuleChecked<FContentBrowserModule>(TEXT("ContentBrowser"));
+
+		FAssetPickerConfig AssetPickerConfig;
+		AssetPickerConfig.Filter.ClassPaths.Add(USkeleton::StaticClass()->GetClassPathName());
+		AssetPickerConfig.OnAssetSelected = FOnAssetSelected::CreateLambda([this](const FAssetData& AssetData)
+			{
+				SkeletonAsset = AssetData;
+			});
+		AssetPickerConfig.bAllowNullSelection = false;
+		AssetPickerConfig.InitialAssetViewType = EAssetViewType::List;
+		AssetPickerConfig.InitialAssetSelection = SkeletonAsset;
+		AssetPickerConfig.HiddenColumnNames =
+		{
+			"DiskSize",
+			"AdditionalPreviewSkeletalMeshes",
+			"PreviewSkeletalMesh"
+		};
+		AssetPickerConfig.bShowPathInColumnView = false;
+		AssetPickerConfig.bShowTypeInColumnView = false;
+		AssetPickerConfig.bFocusSearchBoxWhenOpened = false;
+	
+		return SNew(SVerticalBox)
+			+ SVerticalBox::Slot()
+				.FillHeight(1.0f)
+				.Padding(5.0f)
+				[
+					ContentBrowserModule.Get().CreateAssetPicker(AssetPickerConfig)
+				];
+	}
+	
+	void CloseDialog(bool bWasPicked=false)
+	{
+		bOkClicked = bWasPicked;
+		if (PickerWindow.IsValid())
+		{
+			PickerWindow.Pin()->RequestDestroyWindow();
+		}
+	}
+
+	FReply OnCreateClicked()
+	{
+		Factory->Skeleton = Cast<USkeleton>(SkeletonAsset.GetAsset());
+		Factory->SkeletonTemplate = Cast<USkeletonTemplate>(SkeletonTemplateAsset.GetAsset());
+		
+		CloseDialog(true);
+		return FReply::Handled();
+	}
+	
+	FReply OnCancelClicked()
+	{
+		CloseDialog();
+		return FReply::Handled();
+	}
+
+	virtual FReply OnKeyDown(const FGeometry& MyGeometry, const FKeyEvent& InKeyEvent) override
+	{
+		if (InKeyEvent.GetKey() == EKeys::Escape)
+		{
+			CloseDialog();
+			return FReply::Handled();
+		}
+		return SWidget::OnKeyDown(MyGeometry, InKeyEvent);
+	}
+
+private:
+	TObjectPtr<USkeletonBindingFactory> Factory;
+
+	TWeakPtr<SWindow> PickerWindow;
+
+	FAssetData SkeletonTemplateAsset;
+
+	FAssetData SkeletonAsset;
+	
+	bool bOkClicked;
+};
 
 USkeletonBindingFactory::USkeletonBindingFactory()
 {
@@ -13,7 +237,8 @@ USkeletonBindingFactory::USkeletonBindingFactory()
 
 bool USkeletonBindingFactory::ConfigureProperties()
 {
-	return true;
+	const TSharedRef<SSkeletonBindingCreateDialog> Dialog = SNew(SSkeletonBindingCreateDialog);
+	return Dialog->ConfigureProperties(this);
 }
 
 UObject* USkeletonBindingFactory::FactoryCreateNew(UClass* Class, UObject* InParent, FName Name, EObjectFlags Flags, UObject* Context, FFeedbackContext* Warn, FName CallingContext)
@@ -25,6 +250,9 @@ UObject* USkeletonBindingFactory::FactoryCreateNew(UClass* Class, UObject* InPar
 	}
 
 	USkeletonBinding* NewBinding = NewObject<USkeletonBinding>(InParent, Class, Name, FlagsToUse);
+	NewBinding->InitializeFrom(SkeletonTemplate, Skeleton);
 
 	return NewBinding;
 }
+
+#undef LOCTEXT_NAMESPACE
\ No newline at end of file
diff --git a/Engine/Plugins/Experimental/Animation/SkeletonTemplateFramework/Source/Editor/Private/SkeletonTemplateFramework/SkeletonBindingFactory.h b/Engine/Plugins/Experimental/Animation/SkeletonTemplateFramework/Source/Editor/Private/SkeletonTemplateFramework/SkeletonBindingFactory.h
--- a/Engine/Plugins/Experimental/Animation/SkeletonTemplateFramework/Source/Editor/Private/SkeletonTemplateFramework/SkeletonBindingFactory.h
+++ b/Engine/Plugins/Experimental/Animation/SkeletonTemplateFramework/Source/Editor/Private/SkeletonTemplateFramework/SkeletonBindingFactory.h
@@ -7,6 +7,9 @@
 
 #include "SkeletonBindingFactory.generated.h"
 
+class USkeletonTemplate;
+class USkeleton;
+
 UCLASS()
 class USkeletonBindingFactory : public UFactory
 {
@@ -20,4 +23,8 @@ class USkeletonBindingFactory : public UFactory
 	{
 		return FactoryCreateNew(Class, InParent, Name, Flags, Context, Warn, NAME_None);
 	}
+
+public:
+	TObjectPtr<USkeletonTemplate> SkeletonTemplate;
+	TObjectPtr<USkeleton> Skeleton;
 };
diff --git a/Engine/Plugins/Experimental/Animation/SkeletonTemplateFramework/Source/Runtime/Private/SkeletonTemplateFramework/SkeletonBinding.cpp b/Engine/Plugins/Experimental/Animation/SkeletonTemplateFramework/Source/Runtime/Private/SkeletonTemplateFramework/SkeletonBinding.cpp
--- a/Engine/Plugins/Experimental/Animation/SkeletonTemplateFramework/Source/Runtime/Private/SkeletonTemplateFramework/SkeletonBinding.cpp
+++ b/Engine/Plugins/Experimental/Animation/SkeletonTemplateFramework/Source/Runtime/Private/SkeletonTemplateFramework/SkeletonBinding.cpp
@@ -2,9 +2,30 @@
 
 #include "SkeletonTemplateFramework/SkeletonBinding.h"
 
+#include "Animation/Skeleton.h"
+
 #include UE_INLINE_GENERATED_CPP_BY_NAME(SkeletonBinding)
 
 USkeletonBinding::USkeletonBinding(const FObjectInitializer& ObjectInitializer)
 	: Super(ObjectInitializer)
 {
 }
+
+bool USkeletonBinding::InitializeFrom(const TObjectPtr<USkeletonTemplate> InTemplate, const TObjectPtr<USkeleton> InSkeleton)
+{
+	if (!ensure(InTemplate))
+	{
+		return false;
+	}
+
+	Template = InTemplate;
+	Skeleton = InSkeleton;
+	AttributeBindings.Reset();
+	NamedAttributes = InTemplate->GetNamedAttributes();
+	NamedAttributeSets = InTemplate->GetNamedAttributeSets();
+	NamedAttributeMappings = InTemplate->GetNamedAttributeMappings();
+	CachedSkeletonGuid = InSkeleton->GetGuid();
+	// TODO: Set CachedParentTemplateHash
+
+	return true;
+}
\ No newline at end of file
diff --git a/Engine/Plugins/Experimental/Animation/SkeletonTemplateFramework/Source/Runtime/Private/SkeletonTemplateFramework/SkeletonTemplate.cpp b/Engine/Plugins/Experimental/Animation/SkeletonTemplateFramework/Source/Runtime/Private/SkeletonTemplateFramework/SkeletonTemplate.cpp
--- a/Engine/Plugins/Experimental/Animation/SkeletonTemplateFramework/Source/Runtime/Private/SkeletonTemplateFramework/SkeletonTemplate.cpp
+++ b/Engine/Plugins/Experimental/Animation/SkeletonTemplateFramework/Source/Runtime/Private/SkeletonTemplateFramework/SkeletonTemplate.cpp
@@ -248,6 +248,11 @@ bool USkeletonTemplate::RenameNamedAttributeSet(const FName AttributeSetName, co
 	return true;
 }
 
+const TArray<FSkeletonNamedAttributeMapping>& USkeletonTemplate::GetNamedAttributeMappings() const
+{
+	return NamedAttributeMappings;
+}
+
 FSkeletonNamedAttribute* USkeletonTemplate::FindMutableNamedAttribute(const FName AttributeName)
 {
 	return NamedAttributes.FindByPredicate([AttributeName](const FSkeletonNamedAttribute& Candidate)
diff --git a/Engine/Plugins/Experimental/Animation/SkeletonTemplateFramework/Source/Runtime/Public/SkeletonTemplateFramework/SkeletonBinding.h b/Engine/Plugins/Experimental/Animation/SkeletonTemplateFramework/Source/Runtime/Public/SkeletonTemplateFramework/SkeletonBinding.h
--- a/Engine/Plugins/Experimental/Animation/SkeletonTemplateFramework/Source/Runtime/Public/SkeletonTemplateFramework/SkeletonBinding.h
+++ b/Engine/Plugins/Experimental/Animation/SkeletonTemplateFramework/Source/Runtime/Public/SkeletonTemplateFramework/SkeletonBinding.h
@@ -35,6 +35,9 @@ class USkeletonBinding : public UObject
 public:
 	USkeletonBinding(const FObjectInitializer& ObjectInitializer);
 
+	// Sets up the named attributes using the provided template
+	SKELETONTEMPLATEFRAMEWORKRUNTIME_API bool InitializeFrom(const TObjectPtr<USkeletonTemplate> InTemplate, const TObjectPtr<USkeleton> InSkeleton);
+	
 protected:
 	// The skeleton template to bind
 	// We inherit everything from our parent template, allowing us to override things as needed
diff --git a/Engine/Plugins/Experimental/Animation/SkeletonTemplateFramework/Source/Runtime/Public/SkeletonTemplateFramework/SkeletonTemplate.h b/Engine/Plugins/Experimental/Animation/SkeletonTemplateFramework/Source/Runtime/Public/SkeletonTemplateFramework/SkeletonTemplate.h
--- a/Engine/Plugins/Experimental/Animation/SkeletonTemplateFramework/Source/Runtime/Public/SkeletonTemplateFramework/SkeletonTemplate.h
+++ b/Engine/Plugins/Experimental/Animation/SkeletonTemplateFramework/Source/Runtime/Public/SkeletonTemplateFramework/SkeletonTemplate.h
@@ -108,6 +108,8 @@ public:
 
 	SKELETONTEMPLATEFRAMEWORKRUNTIME_API bool RenameNamedAttributeSet(const FName AttributeSetName, const FName NewAttributeSetName);
 
+	SKELETONTEMPLATEFRAMEWORKRUNTIME_API const TArray<FSkeletonNamedAttributeMapping>& GetNamedAttributeMappings() const;
+
 protected:
 	FSkeletonNamedAttribute* FindMutableNamedAttribute(const FName AttributeName);
 
