Fix LOCTEXT conflict

diff --git a/Engine/Plugins/Experimental/Animation/SkeletonTemplateFramework/Source/Editor/Private/SkeletonTemplateFramework/SkeletonBindingEditorToolkit.cpp b/Engine/Plugins/Experimental/Animation/SkeletonTemplateFramework/Source/Editor/Private/SkeletonTemplateFramework/SkeletonBindingEditorToolkit.cpp
--- a/Engine/Plugins/Experimental/Animation/SkeletonTemplateFramework/Source/Editor/Private/SkeletonTemplateFramework/SkeletonBindingEditorToolkit.cpp
+++ b/Engine/Plugins/Experimental/Animation/SkeletonTemplateFramework/Source/Editor/Private/SkeletonTemplateFramework/SkeletonBindingEditorToolkit.cpp
@@ -94,7 +94,7 @@ void FSkeletonBindingEditorToolkit::RegisterTabSpawners(const TSharedRef<class F
 	
 	InTabManager->RegisterTabSpawner(FSkeletonBindingEditorTabs::AttributeMappingsId, FOnSpawnTab::CreateSP(this, &FSkeletonBindingEditorToolkit::SpawnTabAttributeMappings))
 		.SetDisplayName(LOCTEXT("AttributeMappingsTabMenu_Description", "Attribute Mappings"))
-		.SetTooltipText(LOCTEXT("AttributeMappingsTabMenu_Description", "Shows the attribute mappings view panel"))
+		.SetTooltipText(LOCTEXT("AttributeMappingsTabMenu_Tooltip", "Shows the attribute mappings view panel"))
 		.SetGroup(WorkspaceMenuCategoryRef)
 		.SetIcon(FSlateIcon(FAppStyle::GetAppStyleSetName(), "LevelEditor.Tabs.Outliner"));
 
