Skeleton Templates Editor - Add named attribute mappings tab

diff --git a/Engine/Plugins/Experimental/Animation/SkeletonTemplateFramework/Source/Editor/Private/SkeletonTemplateFramework/SkeletonTemplateEditorToolkit.cpp b/Engine/Plugins/Experimental/Animation/SkeletonTemplateFramework/Source/Editor/Private/SkeletonTemplateFramework/SkeletonTemplateEditorToolkit.cpp
--- a/Engine/Plugins/Experimental/Animation/SkeletonTemplateFramework/Source/Editor/Private/SkeletonTemplateFramework/SkeletonTemplateEditorToolkit.cpp
+++ b/Engine/Plugins/Experimental/Animation/SkeletonTemplateFramework/Source/Editor/Private/SkeletonTemplateFramework/SkeletonTemplateEditorToolkit.cpp
@@ -7,6 +7,7 @@
 #include "SkeletonTemplateFramework/SkeletonTemplate.h"
 #include "SkeletonTemplateFramework/SkeletonTemplateNamedAttributesEditor.h"
 #include "SkeletonTemplateFramework/SkeletonTemplateNamedAttributeSetsEditor.h"
+#include "SkeletonTemplateFramework/SkeletonTemplateNamedAttributeMappingsEditor.h"
 #include "Widgets/Docking/SDockTab.h"
 
 #define LOCTEXT_NAMESPACE "SkeletonTemplateEditorToolkit"
@@ -17,14 +18,14 @@ namespace UE::Anim::STF
 	{
 		static const FName AttributesId;
 		static const FName AttributeSetsId;
-		static const FName TableViewId;
+		static const FName AttributeMappingsId;
 		static const FName DetailsId;
 	};
 }
 
 const FName UE::Anim::STF::FEditorTabs::AttributesId(TEXT("Attributes"));
 const FName UE::Anim::STF::FEditorTabs::AttributeSetsId(TEXT("AttributeSets"));
-const FName UE::Anim::STF::FEditorTabs::TableViewId(TEXT("TableView"));
+const FName UE::Anim::STF::FEditorTabs::AttributeMappingsId(TEXT("AttributeMappings"));
 const FName UE::Anim::STF::FEditorTabs::DetailsId(TEXT("Details"));
 
 void FSkeletonTemplateEditorToolkit::InitEditor(const TArray<UObject*>& InObjects)
@@ -60,7 +61,7 @@ void FSkeletonTemplateEditorToolkit::InitEditor(const TArray<UObject*>& InObject
 			(
 				FTabManager::NewStack()
 				->SetSizeCoefficient(0.5f)
-				->AddTab(FEditorTabs::TableViewId, ETabState::OpenedTab)
+				->AddTab(FEditorTabs::AttributeMappingsId, ETabState::OpenedTab)
 			)
 			->Split
 			(
@@ -94,10 +95,11 @@ void FSkeletonTemplateEditorToolkit::RegisterTabSpawners(const TSharedRef<class
 		.SetGroup(WorkspaceMenuCategoryRef)
 		.SetIcon(FSlateIcon(FAppStyle::GetAppStyleSetName(), "LevelEditor.Tabs.Outliner"));
 	
-	InTabManager->RegisterTabSpawner(FEditorTabs::TableViewId, FOnSpawnTab::CreateSP(this, &FSkeletonTemplateEditorToolkit::SpawnTabTableView))
-		.SetDisplayName(LOCTEXT("TableViewTabMenu_Description", "Table View"))
-		.SetTooltipText(LOCTEXT("TableViewTabMenu_ToolTip", "Shows the table view panel"))
-		.SetGroup(WorkspaceMenuCategoryRef);
+	InTabManager->RegisterTabSpawner(FEditorTabs::AttributeMappingsId, FOnSpawnTab::CreateSP(this, &FSkeletonTemplateEditorToolkit::SpawnTabAttributeMappings))
+		.SetDisplayName(LOCTEXT("AttributeMappingsTabMenu_Description", "Attribute Mappings"))
+		.SetTooltipText(LOCTEXT("AttributeMappingsTabMenu_Description", "Shows the attribute mappings view panel"))
+		.SetGroup(WorkspaceMenuCategoryRef)
+		.SetIcon(FSlateIcon(FAppStyle::GetAppStyleSetName(), "LevelEditor.Tabs.Outliner"));
 
 	InTabManager->RegisterTabSpawner(FEditorTabs::DetailsId, FOnSpawnTab::CreateSP(this, &FSkeletonTemplateEditorToolkit::SpawnTabDetails))
 		.SetDisplayName(LOCTEXT("DetailsTabMenu_Description", "Details"))
@@ -113,7 +115,8 @@ void FSkeletonTemplateEditorToolkit::UnregisterTabSpawners(const TSharedRef<clas
 	FAssetEditorToolkit::UnregisterTabSpawners(InTabManager);
 
 	InTabManager->UnregisterTabSpawner(FEditorTabs::AttributesId);
-	InTabManager->UnregisterTabSpawner(FEditorTabs::TableViewId);
+	InTabManager->UnregisterTabSpawner(FEditorTabs::AttributeSetsId);
+	InTabManager->UnregisterTabSpawner(FEditorTabs::AttributeMappingsId);
 	InTabManager->UnregisterTabSpawner(FEditorTabs::DetailsId);
 }
 
@@ -137,6 +140,18 @@ FLinearColor FSkeletonTemplateEditorToolkit::GetWorldCentricTabColorScale() cons
 	return FLinearColor::White;
 }
 
+void FSkeletonTemplateEditorToolkit::SetDetailsObject(TObjectPtr<UObject> InObject)
+{
+	if (InObject)
+	{
+		DetailsView->SetObjects(TArray{ InObject });
+	}
+	else
+	{
+		DetailsView->SetObjects(TArray<TObjectPtr<UObject>>{ SkeletonTemplate });
+	}
+}
+
 TSharedRef<SDockTab> FSkeletonTemplateEditorToolkit::SpawnTabAttributes(const FSpawnTabArgs& Args)
 {
 	return SNew(SDockTab)
@@ -150,12 +165,22 @@ TSharedRef<SDockTab> FSkeletonTemplateEditorToolkit::SpawnTabAttributeSets(const
 	return SNew(SDockTab)
 		[
 			SAssignNew(AttributeSetsTreeView, UE::Anim::STF::SAttributeSetsTreeView, SkeletonTemplate)
+				.OnNamedAttributeSetsChanged_Lambda([this]()
+				{
+					if (AttributeMappingsTreeView)
+					{
+						AttributeMappingsTreeView->OnNamedAttributeSetsChanged();
+					}
+				})
 		];
 }
 
-TSharedRef<SDockTab> FSkeletonTemplateEditorToolkit::SpawnTabTableView(const FSpawnTabArgs& Args)
+TSharedRef<SDockTab> FSkeletonTemplateEditorToolkit::SpawnTabAttributeMappings(const FSpawnTabArgs& Args)
 {
-	return SNew(SDockTab);
+	return SNew(SDockTab)
+		[
+			SAssignNew(AttributeMappingsTreeView, UE::Anim::STF::SAttributeMappingsTreeView, SkeletonTemplate, this)
+		];
 }
 
 TSharedRef<SDockTab> FSkeletonTemplateEditorToolkit::SpawnTabDetails(const FSpawnTabArgs& Args)
@@ -164,6 +189,11 @@ TSharedRef<SDockTab> FSkeletonTemplateEditorToolkit::SpawnTabDetails(const FSpaw
 	
 	FDetailsViewArgs DetailsViewArgs;
 	DetailsViewArgs.NameAreaSettings = FDetailsViewArgs::HideNameArea;
+	DetailsViewArgs.bAllowSearch = false;
+	DetailsViewArgs.bUpdatesFromSelection = false;
+	DetailsViewArgs.bShowOptions = false;
+	DetailsViewArgs.bShowObjectLabel = false;
+	DetailsViewArgs.bAllowFavoriteSystem = false;
 	
 	DetailsView = PropertyEditorModule.CreateDetailView(DetailsViewArgs);
 	DetailsView->SetObjects(TArray<UObject*>{ SkeletonTemplate });
diff --git a/Engine/Plugins/Experimental/Animation/SkeletonTemplateFramework/Source/Editor/Private/SkeletonTemplateFramework/SkeletonTemplateEditorToolkit.h b/Engine/Plugins/Experimental/Animation/SkeletonTemplateFramework/Source/Editor/Private/SkeletonTemplateFramework/SkeletonTemplateEditorToolkit.h
--- a/Engine/Plugins/Experimental/Animation/SkeletonTemplateFramework/Source/Editor/Private/SkeletonTemplateFramework/SkeletonTemplateEditorToolkit.h
+++ b/Engine/Plugins/Experimental/Animation/SkeletonTemplateFramework/Source/Editor/Private/SkeletonTemplateFramework/SkeletonTemplateEditorToolkit.h
@@ -11,9 +11,18 @@ namespace UE::Anim::STF
 {
 	class SAttributesTreeView;
 	class SAttributeSetsTreeView;
+	class SAttributeMappingsTreeView;
 }
 
-class FSkeletonTemplateEditorToolkit : public FAssetEditorToolkit
+class ISkeletonTemplateEditorToolkit
+{
+public:
+	virtual ~ISkeletonTemplateEditorToolkit() = default;
+	
+	virtual void SetDetailsObject(TObjectPtr<UObject> InObject) = 0;
+};
+
+class FSkeletonTemplateEditorToolkit : public FAssetEditorToolkit, public ISkeletonTemplateEditorToolkit
 {
 public:
 	void InitEditor(const TArray<UObject*>& InObjects);
@@ -26,17 +35,22 @@ public:
 	FString GetWorldCentricTabPrefix() const override;
 	FLinearColor GetWorldCentricTabColorScale() const override;
 
+	// ISkeletonTemplateEditorToolkit
+	void SetDetailsObject(TObjectPtr<UObject> InObject) override;
+	
 private:
 	TSharedRef<SDockTab> SpawnTabAttributes(const FSpawnTabArgs& Args);
 	TSharedRef<SDockTab> SpawnTabAttributeSets(const FSpawnTabArgs& Args);
-	TSharedRef<SDockTab> SpawnTabTableView(const FSpawnTabArgs& Args);
+	TSharedRef<SDockTab> SpawnTabAttributeMappings(const FSpawnTabArgs& Args);
 	TSharedRef<SDockTab> SpawnTabDetails(const FSpawnTabArgs& Args);
 
 	TObjectPtr<USkeletonTemplate> SkeletonTemplate;
 
-	TSharedPtr<IDetailsView> DetailsView;
-
 	TSharedPtr<UE::Anim::STF::SAttributesTreeView> AttributesTreeView;
 
 	TSharedPtr<UE::Anim::STF::SAttributeSetsTreeView> AttributeSetsTreeView;
+
+	TSharedPtr<UE::Anim::STF::SAttributeMappingsTreeView> AttributeMappingsTreeView;
+
+	TSharedPtr<IDetailsView> DetailsView;
 };
\ No newline at end of file
diff --git a/Engine/Plugins/Experimental/Animation/SkeletonTemplateFramework/Source/Editor/Private/SkeletonTemplateFramework/SkeletonTemplateNamedAttributeMappingsEditor.cpp b/Engine/Plugins/Experimental/Animation/SkeletonTemplateFramework/Source/Editor/Private/SkeletonTemplateFramework/SkeletonTemplateNamedAttributeMappingsEditor.cpp
new file mode 100644
--- /dev/null
+++ b/Engine/Plugins/Experimental/Animation/SkeletonTemplateFramework/Source/Editor/Private/SkeletonTemplateFramework/SkeletonTemplateNamedAttributeMappingsEditor.cpp
@@ -0,0 +1,604 @@
+// Copyright Epic Games, Inc. All Rights Reserved.
+
+#include "SkeletonTemplateNamedAttributeMappingsEditor.h"
+
+#include "Framework/MultiBox/MultiBoxBuilder.h"
+#include "Modules/ModuleManager.h"
+#include "SkeletonTemplateFramework/SkeletonTemplate.h"
+#include "Widgets/Input/SCheckBox.h"
+#include "Widgets/Input/SComboBox.h"
+#include "Widgets/Text/SInlineEditableTextBlock.h"
+#include "Widgets/Views/STreeView.h"
+#include "Animation/Skeleton.h"
+#include "PropertyEditorModule.h"
+#include "SkeletonTemplateEditorToolkit.h"
+
+#define LOCTEXT_NAMESPACE "SkeletonTemplateNamedAttributeMappingsEditor"
+
+namespace UE::Anim::STF
+{
+	struct FAttributeTreeItem;
+
+	struct FAttributePartTreeItem : public SAttributeMappingsTreeView::ITreeItem
+	{
+		TObjectPtr<const UClass> ValueType;
+		TSharedPtr<FAttributeTreeItem> Parent;
+
+		virtual void GetChildren(TArray<TSharedPtr<ITreeItem>>& OutChildren) const override {}	
+	};
+	
+	struct FAttributeTreeItem : public SAttributeMappingsTreeView::ITreeItem
+	{
+		TObjectPtr<const UClass> ValueType;
+		bool bHasValue;
+		TSharedPtr<FAttributeTreeItem> Parent;
+		TArray<TSharedPtr<FAttributeTreeItem>> ChildrenAttributes;
+		TArray<TSharedPtr<FAttributePartTreeItem>> ChildrenParts;
+
+		virtual void GetChildren(TArray<TSharedPtr<ITreeItem>>& OutChildren) const override
+		{
+			OutChildren.Append(ChildrenParts);
+			OutChildren.Append(ChildrenAttributes);
+		}
+	};
+	
+	struct FTemplateNamedAttributeMappingListColumns
+	{
+		static const FName NameId;
+		static const FName SourceSetId;
+	};
+
+	const FName FTemplateNamedAttributeMappingListColumns::NameId(TEXT("Name"));
+	const FName FTemplateNamedAttributeMappingListColumns::SourceSetId(TEXT("SourceSet"));
+	
+	struct FTemplateNamedAttributeMappingsColumns
+	{
+		static const FName NameId;
+		static const FName TypeId;
+		static const FName ValueId;
+	};
+
+	const FName FTemplateNamedAttributeMappingsColumns::NameId(TEXT("Name"));
+	const FName FTemplateNamedAttributeMappingsColumns::TypeId(TEXT("Type"));
+	const FName FTemplateNamedAttributeMappingsColumns::ValueId(TEXT("Value"));
+
+	class SNamedElementMappingTableRow : public SMultiColumnTableRow<TSharedPtr<SAttributeMappingsTreeView::ITreeItem>>
+	{
+	public:
+		DECLARE_DELEGATE_RetVal_TwoParams(bool /* Success */, FOnRenamed, FName /* Old Name */, FName /* New Name */);
+
+		SLATE_BEGIN_ARGS(SNamedElementMappingTableRow) { }
+			SLATE_ARGUMENT(TSharedPtr<SAttributeMappingsTreeView::ITreeItem>, TreeItem)
+			SLATE_EVENT(FOnRenamed, OnRenamed);
+		SLATE_END_ARGS()
+
+		void OnCommitRename(const FText& InText, ETextCommit::Type CommitInfo)
+		{
+			const FName OldName = TreeItem->AttributeName;
+			const FName NewName = FName(InText.ToString());
+
+			if (OnRenamed.IsBound())
+			{
+				const bool bSuccess = OnRenamed.Execute(OldName, NewName);
+				if (bSuccess)
+				{
+					TreeItem->AttributeName = NewName;
+				}
+			}
+		}
+
+		void Construct(const FArguments& InArgs, const TSharedRef<STableViewBase>& InOwnerTableView, const TSharedPtr<SAttributeMappingsTreeView> InTreeView)
+		{
+			TreeItem = InArgs._TreeItem;
+			OnRenamed = InArgs._OnRenamed;
+			TreeView = InTreeView;
+
+			FSuperRowType::Construct(
+				FSuperRowType::FArguments(),
+				InOwnerTableView);
+		}
+
+		TSharedRef<SWidget> GenerateWidgetForColumn(const FName& ColumnName) override
+		{
+			if (ColumnName == FTemplateNamedAttributeMappingsColumns::NameId)
+			{
+				return SNew(SHorizontalBox)
+					+ SHorizontalBox::Slot()
+						.AutoWidth()
+						[
+							SNew(SExpanderArrow, SharedThis(this))
+								.ShouldDrawWires(true)
+						]
+					+ SHorizontalBox::Slot()
+						.AutoWidth()
+						[
+							SNew(STextBlock)
+								.Text_Lambda([this]()
+									{
+										return FText::FromName(TreeItem->AttributeName);
+									})
+						];
+			}
+				
+			if (ColumnName == FTemplateNamedAttributeMappingsColumns::TypeId)
+			{
+				return SNew(SHorizontalBox)
+					+ SHorizontalBox::Slot()
+						.VAlign(VAlign_Center)
+						[
+							SNew(STextBlock)
+								.Text(TreeItem->AttributeType)
+						];
+			}
+				
+			if (ColumnName == FTemplateNamedAttributeMappingsColumns::ValueId)
+			{
+				// Shown in details panel, TODO: Show read-only widget here?
+				return SNullWidget::NullWidget;
+			}
+			
+			return SNullWidget::NullWidget;
+			
+		}
+
+	private:
+		TSharedPtr<SAttributeMappingsTreeView::ITreeItem> TreeItem;
+		TSharedPtr<SAttributeMappingsTreeView> TreeView;
+		FOnRenamed OnRenamed;
+	};
+
+	void SAttributeMappingsTreeView::Construct(const FArguments& InArgs, TObjectPtr<USkeletonTemplate> InSkeletonTemplate, ISkeletonTemplateEditorToolkit* InToolkit)
+	{
+		SkeletonTemplate = InSkeletonTemplate;
+
+		RegenerateNamedAttributeSetOptions();
+
+		FPropertyEditorModule& PropertyEditorModule = FModuleManager::GetModuleChecked<FPropertyEditorModule>("PropertyEditor");
+
+		Toolkit = InToolkit;
+							
+		ChildSlot
+		[
+			SNew(SVerticalBox)
+				+ SVerticalBox::Slot()
+				.MinHeight(200.0f)
+				.AutoHeight()
+				[
+					SAssignNew(ListView, SListView<TSharedPtr<FListItem>>)
+						.SelectionMode(ESelectionMode::Single)
+						.ListItemsSource(&ListItems)
+						.HeaderRow
+						(
+							SNew(SHeaderRow)
+								+ SHeaderRow::Column(FTemplateNamedAttributeMappingListColumns::NameId)
+									.FillWidth(0.5f)
+									.DefaultLabel(LOCTEXT("NameLabel", "Name"))
+								+ SHeaderRow::Column(FTemplateNamedAttributeMappingListColumns::SourceSetId)
+									.FillWidth(0.5f)
+									.DefaultLabel(LOCTEXT("SourceSetLabel", "Source Set"))
+						)
+						.OnGenerateRow(this, &SAttributeMappingsTreeView::ListView_GenerateItemRow)
+						.OnSelectionChanged(this, &SAttributeMappingsTreeView::ListView_OnSelectionChanged)
+						.OnContextMenuOpening(this, &SAttributeMappingsTreeView::ListView_HandleContextMenuOpening)
+						.OnItemScrolledIntoView(this, &SAttributeMappingsTreeView::ListView_OnItemScrolledIntoView)
+				]
+				+ SVerticalBox::Slot()
+				.FillHeight(1.0f)
+				[
+					SAssignNew(TreeView, STreeView<TSharedPtr<ITreeItem>>)
+						.SelectionMode(ESelectionMode::Single)
+						.TreeItemsSource(&RootItems)
+						.OnGenerateRow(this, &SAttributeMappingsTreeView::TreeView_GenerateItemRow)
+						.OnSelectionChanged(this, &SAttributeMappingsTreeView::TreeView_OnSelectionChanged)
+						.OnGetChildren(this, &SAttributeMappingsTreeView::TreeView_HandleGetChildren)
+						.OnContextMenuOpening(this, &SAttributeMappingsTreeView::TreeView_HandleContextMenuOpening)
+						.HighlightParentNodesForSelection(true)
+						.HeaderRow
+						(
+							SNew(SHeaderRow)
+								+SHeaderRow::Column(FTemplateNamedAttributeMappingsColumns::NameId)
+									.FillWidth(0.5f)
+									.DefaultLabel(LOCTEXT("NameLabel", "Name"))
+								/*
+								+SHeaderRow::Column(UE::Anim::STF::FNamedAttributeMappingsColumns::TypeId)
+									.FillWidth(0.1f)
+									.DefaultLabel(LOCTEXT("TypeLabel", "Type"))
+								*/
+								+SHeaderRow::Column(FTemplateNamedAttributeMappingsColumns::ValueId)
+									.FillWidth(0.5f)
+									.DefaultLabel(LOCTEXT("ValueLabel", "Value"))
+						)
+				]
+		];
+
+		RegenerateListViewItems();
+		RegenerateTreeViewItems();
+
+		// Expand all tree items on construction
+		for (TSharedPtr<ITreeItem> TreeItem : GetAllTreeItems())
+		{
+			TreeView->SetItemExpansion(TreeItem, true);
+		}
+	}
+	
+	void SAttributeMappingsTreeView::OnNamedAttributeSetsChanged()
+	{
+		RegenerateNamedAttributeSetOptions();
+		RegenerateListViewItems();
+	}
+
+	class SListViewRow : public SMultiColumnTableRow<TSharedPtr<SAttributeMappingsTreeView::FListItem>>
+	{
+	public:
+		DECLARE_DELEGATE_RetVal_OneParam(bool /* Success */, FOnSourceSetChanged, FName /* New Source Set */);
+		DECLARE_DELEGATE_RetVal_TwoParams(bool /* Success */, FOnRenamed, FName /* Old Name */, FName /* New Name */);
+
+		SLATE_BEGIN_ARGS(SListViewRow) { }
+			SLATE_ARGUMENT(TSharedPtr<SAttributeMappingsTreeView::FListItem>, ListItem)
+			SLATE_ARGUMENT(TArray<TSharedPtr<FName>>*, SourceSetOptions)
+			SLATE_EVENT(FOnRenamed, OnRenamed);
+			SLATE_EVENT(FOnSourceSetChanged, OnSourceSetChanged);
+			
+		SLATE_END_ARGS()
+
+		void Construct(const FArguments& InArgs, const TSharedRef<STableViewBase>& InOwnerTableView)
+		{
+			ListItem = InArgs._ListItem;
+			OnRenamed = InArgs._OnRenamed;
+			OnSourceSetChanged = InArgs._OnSourceSetChanged;
+			SourceSetOptions = InArgs._SourceSetOptions;
+
+			FSuperRowType::Construct(
+				FSuperRowType::FArguments(),
+				InOwnerTableView);
+		}
+
+		TSharedRef<SWidget> GenerateWidgetForColumn(const FName& ColumnName) override
+		{
+			if (ColumnName == FTemplateNamedAttributeMappingListColumns::NameId)
+			{
+				TSharedPtr<SInlineEditableTextBlock> InlineWidget;
+
+				TSharedRef<SInlineEditableTextBlock> Row = SAssignNew(InlineWidget, SInlineEditableTextBlock)
+					.Text_Lambda([this]()
+					{
+						return FText::FromName(ListItem->Name);
+					})
+					.OnTextCommitted_Lambda([this](const FText& InText, ETextCommit::Type CommitInfo)
+					{
+						const FName OldName = ListItem->Name;
+						const FName NewName = FName(InText.ToString());
+
+						if (OnRenamed.IsBound())
+						{
+							const bool bSuccess = OnRenamed.Execute(OldName, NewName);
+							if (bSuccess)
+							{
+								ListItem->Name = NewName;
+							}
+						}
+					});
+				
+				ListItem->OnRenameRequested.BindSP(InlineWidget.Get(), &SInlineEditableTextBlock::EnterEditingMode);
+		
+				return Row;
+			}
+
+			if (ColumnName == FTemplateNamedAttributeMappingListColumns::SourceSetId)
+			{
+				TSharedPtr<FName> InitialValue;
+				for (TSharedPtr<FName> SourceSetName : *SourceSetOptions)
+				{
+					if (*SourceSetName == ListItem->SourceSetName)
+					{
+						InitialValue = SourceSetName;
+						break;
+					}
+				}
+				
+				return SNew(SComboBox<TSharedPtr<FName>>)
+					.InitiallySelectedItem(InitialValue)
+					.OptionsSource(SourceSetOptions)
+					.OnGenerateWidget_Lambda([](TSharedPtr<FName> SourceSetName)
+						{
+							return SNew(STextBlock)
+								.Text(FText::FromName(*SourceSetName));
+						})
+					.OnSelectionChanged_Lambda([this](TSharedPtr<FName> NewSelection, ESelectInfo::Type)
+						{
+							if (NewSelection && OnSourceSetChanged.IsBound())
+							{
+								const bool bSuccess = OnSourceSetChanged.Execute(*NewSelection);
+								if (bSuccess)
+								{
+									ListItem->SourceSetName = *NewSelection;
+								}
+							}
+						})
+					[
+						SNew(STextBlock)
+							.Text_Lambda([this]()
+							{
+								return FText::FromName(ListItem->SourceSetName);
+							})
+					];
+			}
+
+			return SNullWidget::NullWidget;
+		}
+
+	private:
+		FOnSourceSetChanged OnSourceSetChanged;
+		
+		TArray<TSharedPtr<FName>>* SourceSetOptions;
+		
+		TSharedPtr<SAttributeMappingsTreeView::FListItem> ListItem;
+		FOnRenamed OnRenamed;
+	};
+	
+	TSharedRef<ITableRow> SAttributeMappingsTreeView::ListView_GenerateItemRow(TSharedPtr<FListItem> Item, const TSharedRef<STableViewBase>& OwnerTable)
+	{
+		return SNew(SListViewRow, OwnerTable)
+			.ListItem(Item)
+			.OnRenamed_Lambda([this](const FName OldAttributeSetName, const FName NewAttributeSetName)
+			{
+				return SkeletonTemplate->RenameNamedAttributeMapping(OldAttributeSetName, NewAttributeSetName);
+			})
+			.OnSourceSetChanged_Lambda([this, Item](const FName NewSourceSetName)
+			{
+				const bool bSuccess = SkeletonTemplate->SetNamedAttributeMappingSourceSet(Item->Name, NewSourceSetName);
+				if (SelectedMappingName == Item->Name)
+				{
+					RegenerateTreeViewItems();
+				}
+				return bSuccess;
+			})
+			.SourceSetOptions(&SourceSetOptions);
+	}
+	
+	void SAttributeMappingsTreeView::ListView_OnItemScrolledIntoView(TSharedPtr<FListItem> InItem, const TSharedPtr<ITableRow>& InWidget)
+	{
+		if (ListView_DeferredRenameRequest.IsValid())
+		{
+			ListView_DeferredRenameRequest->OnRenameRequested.ExecuteIfBound();
+			ListView_DeferredRenameRequest.Reset();
+		}
+	}
+
+	void SAttributeMappingsTreeView::ListView_OnSelectionChanged(TSharedPtr<FListItem> InItem, ESelectInfo::Type InSelectInfo)
+	{
+		SelectedMappingName = InItem ? InItem->Name : NAME_None;
+		RegenerateTreeViewItems();
+	}
+	
+	TSharedPtr<SWidget> SAttributeMappingsTreeView::ListView_HandleContextMenuOpening()
+	{
+		FMenuBuilder MenuBuilder(true, nullptr);
+
+		const TArray<TSharedPtr<FListItem>> Selection = ListView->GetSelectedItems();
+		if (Selection.IsEmpty())
+		{
+			MenuBuilder.AddMenuEntry(
+				LOCTEXT("AddNamedAttributeMapping", "Add Named Attribute Mapping"),
+				LOCTEXT("AddNamedAttributeMapping_Tooltip", "Create a new named attribute mapping"),
+				FSlateIcon(FAppStyle::GetAppStyleSetName(), "Icons.Plus"),
+				FUIAction(
+					FExecuteAction::CreateLambda([this]()
+					{
+						SkeletonTemplate->AddAttributeMapping("NewAttributeMapping");
+						RegenerateListViewItems();
+
+						for (const TSharedPtr<FListItem>& Item : ListItems)
+						{
+							if (Item->Name == "NewAttributeMapping")
+							{
+								ListView->RequestScrollIntoView(Item);
+								ListView_DeferredRenameRequest = Item;
+							}
+						}
+					})
+				));
+		}
+		
+		return MenuBuilder.MakeWidget();
+	}
+
+	void SAttributeMappingsTreeView::TreeView_HandleGetChildren(TSharedPtr<ITreeItem> InItem, TArray<TSharedPtr<ITreeItem>>& OutChildren) const
+	{
+		InItem->GetChildren(OutChildren);
+	}
+
+	void SAttributeMappingsTreeView::TreeView_OnSelectionChanged(TSharedPtr<ITreeItem> InItem, ESelectInfo::Type InSelectInfo)
+	{
+		if (InItem)
+		{			
+			const FSkeletonNamedAttributeMapping* const Mapping = SkeletonTemplate->FindNamedAttributeMapping(SelectedMappingName);
+			const FSkeletonNamedAttributeMappingEntry* const Entry = Mapping->TableData.FindByPredicate([InItem](const FSkeletonNamedAttributeMappingEntry& Candidate)
+			{
+				return Candidate.AttributeName == InItem->AttributeName;
+			});
+
+			if (Entry)
+			{
+				Toolkit->SetDetailsObject(Entry->Value);
+			}
+		}
+		else
+		{
+			Toolkit->SetDetailsObject(nullptr);
+		}
+	}
+
+	TSharedRef<ITableRow> SAttributeMappingsTreeView::TreeView_GenerateItemRow(TSharedPtr<SAttributeMappingsTreeView::ITreeItem> Item, const TSharedRef<STableViewBase>& OwnerTable)
+	{
+		return SNew(SNamedElementMappingTableRow, OwnerTable, SharedThis(this))
+			.TreeItem(Item);
+	}
+
+	TSharedPtr<SWidget> SAttributeMappingsTreeView::TreeView_HandleContextMenuOpening()
+	{
+		FMenuBuilder MenuBuilder(true, nullptr);
+		return MenuBuilder.MakeWidget();
+	}
+
+	TArray<TSharedPtr<UE::Anim::STF::SAttributeMappingsTreeView::ITreeItem>> SAttributeMappingsTreeView::GetAllTreeItems()
+	{
+		using namespace UE::Anim::STF;
+
+		TArray<TSharedPtr<SAttributeMappingsTreeView::ITreeItem>> AllItems;
+		AllItems.Append(RootItems);
+	
+		for (int32 Index = 0; Index < AllItems.Num(); ++Index)
+		{
+			AllItems[Index]->GetChildren(AllItems);
+		}
+
+		return AllItems;
+	}
+
+	void SAttributeMappingsTreeView::RegenerateNamedAttributeSetOptions()
+	{
+		SourceSetOptions.Reset();
+		
+		SourceSetOptions.Add(MakeShared<FName>(NAME_None));
+		for (const FSkeletonNamedAttributeSet& AttributeSet : SkeletonTemplate->GetNamedAttributeSets())
+		{
+			SourceSetOptions.Add(MakeShared<FName>(AttributeSet.Name));
+		}
+	}
+
+	void SAttributeMappingsTreeView::RegenerateListViewItems()
+	{
+		ListItems.Reset();
+
+		const TArray<FSkeletonNamedAttributeMapping>& Mappings = SkeletonTemplate->GetNamedAttributeMappings();
+		for (const FSkeletonNamedAttributeMapping& Mapping : Mappings)
+		{
+			TSharedPtr<FListItem> ListItem = MakeShared<FListItem>();
+			ListItem->Name = Mapping.Name;
+			ListItem->SourceSetName = Mapping.SourceAttributeSet;
+
+			ListItems.Add(ListItem);
+		}
+	}
+
+	void SAttributeMappingsTreeView::RegenerateTreeViewItems()
+	{
+		// Make note of all tree items currently expanded
+		/*
+		TSet<FName> ExpandedAttributeNames;
+		{
+			for (TSharedPtr<FTreeItem> TreeItem : GetAllTreeItems())
+			{
+				if (TreeView->IsItemExpanded(TreeItem))
+				{
+					ExpandedAttributeNames.Add(TreeItem->Name);
+				}
+			}
+		}
+		*/
+
+		// Rebuild items
+		{
+			RootItems.Reset();
+			TMap<FName, TSharedPtr<FAttributeTreeItem>> AttributesItemMap;
+
+			if (const FSkeletonNamedAttributeMapping* const AttributeMapping = SkeletonTemplate->FindNamedAttributeMapping(SelectedMappingName))
+			{
+				TArray<TPair<FSkeletonNamedAttribute, TSharedPtr<FAttributeTreeItem>>> NamedAttributesQueue;
+				
+				for (const FSkeletonNamedAttributeMappingEntry& MappingEntry : AttributeMapping->TableData)
+				{
+					const FSkeletonNamedAttribute* NamedAttribute = SkeletonTemplate->FindNamedAttribute(MappingEntry.AttributeName);
+					check(NamedAttribute);
+					
+					TSharedPtr<FAttributeTreeItem> MappingTreeItem = MakeShared<FAttributeTreeItem>();
+					MappingTreeItem->AttributeName = NamedAttribute->Name;
+					MappingTreeItem->TreeView = this;
+					MappingTreeItem->bHasValue = true;
+
+					const FSkeletonNamedAttributeMappingType* const AttributeMappingType = AttributeMapping->MappingTypes.FindByPredicate([NamedAttribute](const FSkeletonNamedAttributeMappingType& MappingType)
+					{
+						return MappingType.SourceType == NamedAttribute->Type;
+					});
+
+					if (AttributeMappingType)
+					{
+						MappingTreeItem->ValueType = AttributeMappingType->TargetType;
+						MappingTreeItem->AttributeType = FText::Format(LOCTEXT("AttributeType", "{0} -> {1}"), FText::FromName(NamedAttribute->Type.GetFName()), FText::FromName(AttributeMappingType->TargetType.GetFName()));
+					}
+					else
+					{
+						MappingTreeItem->ValueType = nullptr;
+						MappingTreeItem->AttributeType = FText::Format(LOCTEXT("AttributeTypeNone", "{0} -> None"), FText::FromName(NamedAttribute->Type.GetFName()));
+					}
+
+					AttributesItemMap.Add(NamedAttribute->Name, MappingTreeItem);
+					NamedAttributesQueue.Add(TPair<FSkeletonNamedAttribute, TSharedPtr<FAttributeTreeItem>>(*NamedAttribute, MappingTreeItem));
+				}
+
+				while (!NamedAttributesQueue.IsEmpty())
+				{
+					const TPair<FSkeletonNamedAttribute, TSharedPtr<FAttributeTreeItem>> Attribute = NamedAttributesQueue.Pop();
+					const FSkeletonNamedAttribute NamedAttribute = Attribute.Key;
+					TSharedPtr<FAttributeTreeItem> AttributeTreeItem = Attribute.Value;
+
+					FName AncestorName = NamedAttribute.ParentName;
+					TSharedPtr<FAttributeTreeItem> ClosestAncestorItem = nullptr;
+					while (AncestorName != NAME_None)
+					{
+						TSharedPtr<FAttributeTreeItem>* FoundClosestAncestorItem = AttributesItemMap.Find(AncestorName);
+						if (FoundClosestAncestorItem)
+						{
+							ClosestAncestorItem = *FoundClosestAncestorItem;
+							break;
+						}
+
+						if (const FSkeletonNamedAttribute* const Parent = SkeletonTemplate->FindNamedAttribute(AncestorName))
+						{
+							AncestorName = Parent->ParentName;
+						}
+						else
+						{
+							break;
+						}
+					}
+
+					if (ClosestAncestorItem)
+					{
+						AttributeTreeItem->Parent = ClosestAncestorItem;
+						ClosestAncestorItem->ChildrenAttributes.Add(AttributeTreeItem);
+					}
+					else
+					{
+						AttributeTreeItem->Parent = nullptr;
+						RootItems.Add(AttributeTreeItem);
+					}
+				}
+			}
+		}
+
+		// Update tree view and restore tree item expanded states
+		{
+			check(TreeView);
+			TreeView->RequestTreeRefresh();
+
+			for (TSharedPtr<ITreeItem> TreeItem : GetAllTreeItems())
+			{
+				TreeView->SetItemExpansion(TreeItem, true);
+			}
+
+			/*
+			for (TSharedPtr<FTreeItem> TreeItem : GetAllTreeItems())
+			{
+				if (ExpandedAttributeNames.Contains(TreeItem->Name))
+				{
+					TreeView->SetItemExpansion(TreeItem, true);
+				}
+			}
+			*/
+		}
+	}
+}
+
+#undef LOCTEXT_NAMESPACE
\ No newline at end of file
diff --git a/Engine/Plugins/Experimental/Animation/SkeletonTemplateFramework/Source/Editor/Private/SkeletonTemplateFramework/SkeletonTemplateNamedAttributeMappingsEditor.h b/Engine/Plugins/Experimental/Animation/SkeletonTemplateFramework/Source/Editor/Private/SkeletonTemplateFramework/SkeletonTemplateNamedAttributeMappingsEditor.h
new file mode 100644
--- /dev/null
+++ b/Engine/Plugins/Experimental/Animation/SkeletonTemplateFramework/Source/Editor/Private/SkeletonTemplateFramework/SkeletonTemplateNamedAttributeMappingsEditor.h
@@ -0,0 +1,104 @@
+// Copyright Epic Games, Inc. All Rights Reserved.
+
+#pragma once
+
+#include "CoreMinimal.h"
+#include "Widgets/SCompoundWidget.h"
+
+class IDetailsView;
+class ITableRow;
+class STableViewBase;
+template<typename ItemType> class SListView;
+template<typename ItemType> class STreeView;
+class USkeletonTemplate;
+
+class ISkeletonTemplateEditorToolkit;
+
+namespace UE::Anim::STF
+{
+	class SNamedElementMappingTableRow;
+	
+	class SAttributeMappingsTreeView : public SCompoundWidget
+	{
+	public:
+		friend SNamedElementMappingTableRow;
+		
+		SLATE_BEGIN_ARGS(SAttributeMappingsTreeView) { }
+		SLATE_END_ARGS()
+
+		struct FListItem
+		{
+			FName Name;
+			FName SourceSetName;
+
+			DECLARE_DELEGATE(FOnRenameRequested);
+			FOnRenameRequested OnRenameRequested;
+		};
+		
+		struct ITreeItem
+		{
+			virtual ~ITreeItem() = default;
+			
+			FName AttributeName;
+			FText AttributeType;
+			bool bHasValue;
+			SAttributeMappingsTreeView* TreeView;
+			
+			virtual void GetChildren(TArray<TSharedPtr<ITreeItem>>& OutChildren) const = 0;
+		};
+
+		friend ITreeItem;
+		
+		void Construct(const FArguments& InArgs, TObjectPtr<USkeletonTemplate> InSkeletonTemplate, ISkeletonTemplateEditorToolkit* Toolkit);
+
+		void OnNamedAttributeSetsChanged();
+
+	private:
+		TSharedRef<ITableRow> ListView_GenerateItemRow(TSharedPtr<FListItem> Item, const TSharedRef<STableViewBase>& OwnerTable);
+
+		void ListView_OnItemScrolledIntoView(TSharedPtr<FListItem> InItem, const TSharedPtr<ITableRow>& InWidget);
+
+		TSharedPtr<SWidget> ListView_HandleContextMenuOpening();
+		
+		void ListView_OnSelectionChanged(TSharedPtr<FListItem> InItem, ESelectInfo::Type InSelectInfo);
+
+		void TreeView_HandleGetChildren(TSharedPtr<ITreeItem> InItem, TArray<TSharedPtr<ITreeItem>>& OutChildren) const;
+
+		void TreeView_OnSelectionChanged(TSharedPtr<ITreeItem> InItem, ESelectInfo::Type InSelectInfo);
+		
+		TSharedRef<ITableRow> TreeView_GenerateItemRow(TSharedPtr<ITreeItem> Item, const TSharedRef<STableViewBase>& OwnerTable);
+		
+		TSharedPtr<SWidget> TreeView_HandleContextMenuOpening();
+
+	private:
+		TArray<TSharedPtr<ITreeItem>> GetAllTreeItems();
+
+		// Regenerates the cached list of named attribute sets
+		void RegenerateNamedAttributeSetOptions();
+
+		void RegenerateListViewItems();
+		
+		void RegenerateTreeViewItems();
+
+	private:
+		TSharedPtr<SListView<TSharedPtr<FListItem>>> ListView;
+
+		TSharedPtr<FListItem> ListView_DeferredRenameRequest;
+
+		TSharedPtr<STreeView<TSharedPtr<ITreeItem>>> TreeView;
+
+		TArray<TSharedPtr<FListItem>> ListItems;
+		
+		FName SelectedMappingName;
+
+		TArray<TSharedPtr<ITreeItem>> RootItems;
+
+		TSharedPtr<ITreeItem> TreeView_DeferredRenameRequest;
+
+		TObjectPtr<USkeletonTemplate> SkeletonTemplate;
+
+		TArray<TSharedPtr<FName>> SourceSetOptions;
+
+		ISkeletonTemplateEditorToolkit* Toolkit;
+	};
+}
\ No newline at end of file
diff --git a/Engine/Plugins/Experimental/Animation/SkeletonTemplateFramework/Source/Editor/Private/SkeletonTemplateFramework/SkeletonTemplateNamedAttributeSetsEditor.cpp b/Engine/Plugins/Experimental/Animation/SkeletonTemplateFramework/Source/Editor/Private/SkeletonTemplateFramework/SkeletonTemplateNamedAttributeSetsEditor.cpp
--- a/Engine/Plugins/Experimental/Animation/SkeletonTemplateFramework/Source/Editor/Private/SkeletonTemplateFramework/SkeletonTemplateNamedAttributeSetsEditor.cpp
+++ b/Engine/Plugins/Experimental/Animation/SkeletonTemplateFramework/Source/Editor/Private/SkeletonTemplateFramework/SkeletonTemplateNamedAttributeSetsEditor.cpp
@@ -17,14 +17,16 @@
 
 namespace UE::Anim::STF
 {
-	struct FNamedAttributeSetsColumns
+	class FNamedElementDragDropOp;
+
+	struct FTemplateNamedAttributeSetsColumns
 	{
 		static const FName NameId;
 		static const FName TypeId;
 	};
 
-	const FName FNamedAttributeSetsColumns::NameId(TEXT("Name"));
-	const FName FNamedAttributeSetsColumns::TypeId(TEXT("Type"));
+	const FName FTemplateNamedAttributeSetsColumns::NameId(TEXT("Name"));
+	const FName FTemplateNamedAttributeSetsColumns::TypeId(TEXT("Type"));
 
 	class SNamedElementSetTableRow : public SMultiColumnTableRow<TSharedPtr<SAttributeSetsTreeView::FTreeItem>>
 	{
@@ -61,7 +63,7 @@ namespace UE::Anim::STF
 				return FReply::Unhandled();
 			}
 
-			OnAddAttributeToSet.ExecuteIfBound(TreeItem->Name, DragDropOp->Element.Pin()->Name);
+			OnAddAttributeToSet.ExecuteIfBound(TreeItem->Name, DragDropOp->NamedAttribute);
 
 			return FReply::Handled();
 		}
@@ -96,7 +98,7 @@ namespace UE::Anim::STF
 
 		TSharedRef<SWidget> GenerateWidgetForColumn(const FName& ColumnName) override
 		{
-			if (ColumnName == UE::Anim::STF::FNamedAttributeSetsColumns::NameId)
+			if (ColumnName == FTemplateNamedAttributeSetsColumns::NameId)
 			{
 				TSharedPtr<SInlineEditableTextBlock> InlineWidget;
 
@@ -109,6 +111,15 @@ namespace UE::Anim::STF
 						]
 					+ SHorizontalBox::Slot()
 						.AutoWidth()
+						.Padding(4.0f, 0.0f, 0.0f, 0.0f)
+						[
+							SNew(SImage)
+								.Image(FAppStyle::GetBrush("LevelEditor.Tabs.Outliner"))
+								.Visibility(TreeItem->ItemType == SAttributeSetsTreeView::FTreeItem::EType::AttributeSet ? EVisibility::Visible : EVisibility::Collapsed)
+						]
+					+ SHorizontalBox::Slot()
+						.AutoWidth()
+						.Padding(4.0f, 0.0f, 0.0f, 0.0f)
 						[
 							SAssignNew(InlineWidget, SInlineEditableTextBlock)
 								.Text_Lambda([this]()
@@ -122,7 +133,8 @@ namespace UE::Anim::STF
 
 				return HorizontalBox;
 			}
-			else if (ColumnName == UE::Anim::STF::FNamedAttributeSetsColumns::TypeId)
+
+			if (ColumnName == FTemplateNamedAttributeSetsColumns::TypeId)
 			{
 				if (TreeItem->ItemType == SAttributeSetsTreeView::FTreeItem::EType::AttributeSet)
 				{
@@ -137,10 +149,8 @@ namespace UE::Anim::STF
 							});
 				}
 			}
-			else
-			{
-				return SNullWidget::NullWidget;
-			}
+			
+			return SNullWidget::NullWidget;
 		}
 
 	private:
@@ -152,6 +162,7 @@ namespace UE::Anim::STF
 	void SAttributeSetsTreeView::Construct(const FArguments& InArgs, TObjectPtr<USkeletonTemplate> InSkeletonTemplate)
 	{
 		SkeletonTemplate = InSkeletonTemplate;
+		OnNamedAttributeSetsChanged = InArgs._OnNamedAttributeSetsChanged;
 
 		ChildSlot
 		[
@@ -165,10 +176,10 @@ namespace UE::Anim::STF
 				.HeaderRow
 				(
 					SNew(SHeaderRow)
-						+SHeaderRow::Column(UE::Anim::STF::FNamedAttributeSetsColumns::NameId)
+						+SHeaderRow::Column(FTemplateNamedAttributeSetsColumns::NameId)
 							.FillWidth(0.5f)
 							.DefaultLabel(LOCTEXT("NameLabel", "Name"))
-						+SHeaderRow::Column(UE::Anim::STF::FNamedAttributeSetsColumns::TypeId)
+						+SHeaderRow::Column(FTemplateNamedAttributeSetsColumns::TypeId)
 							.FillWidth(0.5f)
 							.DefaultLabel(LOCTEXT("TypeLabel", "Type"))
 				)
@@ -205,13 +216,34 @@ namespace UE::Anim::STF
 				.TreeItem(Item)
 				.OnAddAttributeToSet_Lambda([this, Item](const FName AttributeSet, const FName AttributeName)
 					{
-						SkeletonTemplate->AddAttributeToSet(AttributeSet, AttributeName);
-						TreeView->SetItemExpansion(Item, true);
+						SkeletonTemplate->AddAttributeToSet(AttributeSet, AttributeName);					
 						RegenerateTreeViewItems();
+
+						for (TSharedPtr<FTreeItem> TreeItem : GetAllTreeItems())
+						{
+							if (TreeItem->Name == AttributeName)
+							{
+								TSharedPtr<FTreeItem> ParentTreeItem = TreeItem->Parent;
+								while (ParentTreeItem)
+								{
+									TreeView->SetItemExpansion(ParentTreeItem, true);
+									ParentTreeItem = ParentTreeItem->Parent;
+								}
+								
+								break;
+							}
+						}
+
+						OnNamedAttributeSetsChanged.ExecuteIfBound();
 					})
 				.OnRenamed_Lambda([this](const FName OldName, const FName NewName) -> bool
 				{
-					return SkeletonTemplate->RenameNamedAttributeSet(OldName, NewName);
+					const bool bSuccess = SkeletonTemplate->RenameNamedAttributeSet(OldName, NewName);
+					if (bSuccess)
+					{
+						OnNamedAttributeSetsChanged.ExecuteIfBound();
+					}
+					return bSuccess;
 				});
 		}
 		else
@@ -256,6 +288,7 @@ namespace UE::Anim::STF
 							{
 								SkeletonTemplate->RemoveAttributeSet(SelectedItem->Name);
 								RegenerateTreeViewItems();
+								OnNamedAttributeSetsChanged.ExecuteIfBound();
 							})
 						));
 
@@ -270,16 +303,13 @@ namespace UE::Anim::STF
 						FUIAction(
 							FExecuteAction::CreateLambda([this, SelectedItem]()
 							{
-								SkeletonTemplate->RemoveAttributeFromSet(SelectedItem->Parent->Name, SelectedItem->Name);
+								SkeletonTemplate->RemoveAttributeFromSet(SelectedItem->SetItem->Name, SelectedItem->Name);
 								RegenerateTreeViewItems();
 							})
 						));
 
 					break;
 				}
-				case FTreeItem::EType::AttributePart:
-				{
-				}
 			}
 		}
 		else
@@ -302,6 +332,8 @@ namespace UE::Anim::STF
 								DeferredRenameRequest = Item;
 							}
 						}
+
+						OnNamedAttributeSetsChanged.ExecuteIfBound();
 					})
 				));
 		}
@@ -341,6 +373,7 @@ namespace UE::Anim::STF
 		// Rebuild items
 		{
 			RootItems.Reset();
+			TMap<FName, TSharedPtr<FTreeItem>> AttributesItemMap;
 	
 			for (const FSkeletonNamedAttributeSet& NamedAttributeSet : SkeletonTemplate->GetNamedAttributeSets())
 			{
@@ -349,15 +382,34 @@ namespace UE::Anim::STF
 				SetItem->ItemType = FTreeItem::EType::AttributeSet;
 				SetItem->Type = NAME_None;
 
-				for (const FName AttributeName : NamedAttributeSet.NamedAttributes)
+				for (const FName& AttributeName : NamedAttributeSet.NamedAttributes)
 				{
+					const FSkeletonNamedAttribute* NamedAttribute = SkeletonTemplate->FindNamedAttribute(AttributeName);
+					if (!ensure(NamedAttribute))
+					{
+						continue;
+					}
+
+					TSharedPtr<FTreeItem>* ParentItem = AttributesItemMap.Find(NamedAttribute->ParentName);
+
 					TSharedPtr<FTreeItem> AttributeItem = MakeShared<FTreeItem>();
 					AttributeItem->Name = AttributeName;
 					AttributeItem->ItemType = FTreeItem::EType::Attribute;
-					AttributeItem->Type = SkeletonTemplate->FindNamedAttribute(AttributeName)->Type.GetFName();
-					AttributeItem->Parent = SetItem;
+					AttributeItem->Type = NamedAttribute->Type.GetFName();
+					AttributeItem->SetItem = SetItem;
 
-					SetItem->Children.Add(AttributeItem);
+					if (ParentItem)
+					{
+						AttributeItem->Parent = *ParentItem;
+						(*ParentItem)->Children.Add(AttributeItem);
+					}
+					else
+					{
+						AttributeItem->Parent = SetItem;
+						SetItem->Children.Add(AttributeItem);
+					}
+					
+					AttributesItemMap.Add(AttributeName, AttributeItem);
 				}
 
 				RootItems.Add(SetItem);
diff --git a/Engine/Plugins/Experimental/Animation/SkeletonTemplateFramework/Source/Editor/Private/SkeletonTemplateFramework/SkeletonTemplateNamedAttributeSetsEditor.h b/Engine/Plugins/Experimental/Animation/SkeletonTemplateFramework/Source/Editor/Private/SkeletonTemplateFramework/SkeletonTemplateNamedAttributeSetsEditor.h
--- a/Engine/Plugins/Experimental/Animation/SkeletonTemplateFramework/Source/Editor/Private/SkeletonTemplateFramework/SkeletonTemplateNamedAttributeSetsEditor.h
+++ b/Engine/Plugins/Experimental/Animation/SkeletonTemplateFramework/Source/Editor/Private/SkeletonTemplateFramework/SkeletonTemplateNamedAttributeSetsEditor.h
@@ -15,7 +15,10 @@ namespace UE::Anim::STF
 	class SAttributeSetsTreeView : public SCompoundWidget
 	{
 	public:
+		DECLARE_DELEGATE(FOnNamedAttributeSetsChanged);
+		
 		SLATE_BEGIN_ARGS(SAttributeSetsTreeView) { }
+			SLATE_EVENT(FOnNamedAttributeSetsChanged, OnNamedAttributeSetsChanged);
 		SLATE_END_ARGS()
 
 		struct FTreeItem
@@ -24,12 +27,12 @@ namespace UE::Anim::STF
 			{
 				AttributeSet,
 				Attribute,
-				AttributePart,
 			};
 
 			EType ItemType;
 			FName Name;
 			FName Type;
+			TSharedPtr<FTreeItem> SetItem;
 			TSharedPtr<FTreeItem> Parent;
 			TArray<TSharedPtr<FTreeItem>> Children;
 
@@ -61,5 +64,7 @@ namespace UE::Anim::STF
 		TSharedPtr<FTreeItem> DeferredRenameRequest;
 
 		TObjectPtr<USkeletonTemplate> SkeletonTemplate;
+		
+		FOnNamedAttributeSetsChanged OnNamedAttributeSetsChanged;
 	};
 }
\ No newline at end of file
diff --git a/Engine/Plugins/Experimental/Animation/SkeletonTemplateFramework/Source/Editor/Private/SkeletonTemplateFramework/SkeletonTemplateNamedAttributesEditor.cpp b/Engine/Plugins/Experimental/Animation/SkeletonTemplateFramework/Source/Editor/Private/SkeletonTemplateFramework/SkeletonTemplateNamedAttributesEditor.cpp
--- a/Engine/Plugins/Experimental/Animation/SkeletonTemplateFramework/Source/Editor/Private/SkeletonTemplateFramework/SkeletonTemplateNamedAttributesEditor.cpp
+++ b/Engine/Plugins/Experimental/Animation/SkeletonTemplateFramework/Source/Editor/Private/SkeletonTemplateFramework/SkeletonTemplateNamedAttributesEditor.cpp
@@ -5,8 +5,8 @@
 #include "Framework/MultiBox/MultiBoxBuilder.h"
 #include "Modules/ModuleManager.h"
 #include "SkeletonTemplateFramework/SkeletonTemplate.h"
-#include "StructViewerFilter.h"
-#include "StructViewerModule.h"
+#include "ClassViewerFilter.h"
+#include "ClassViewerModule.h"
 #include "Widgets/Input/SComboButton.h"
 #include "Widgets/Text/SInlineEditableTextBlock.h"
 #include "Widgets/Views/STreeView.h"
@@ -15,19 +15,19 @@
 
 namespace UE::Anim::STF
 {
-	struct FOutlinerColumns
+	struct FTemplateNamedAttributesColumns
 	{
 		static const FName NameId;
 		static const FName TypeId;
 	};
 
-	const FName FOutlinerColumns::NameId(TEXT("Name"));
-	const FName FOutlinerColumns::TypeId(TEXT("Type"));
+	const FName FTemplateNamedAttributesColumns::NameId(TEXT("Name"));
+	const FName FTemplateNamedAttributesColumns::TypeId(TEXT("Type"));
 
-	TSharedRef<FNamedElementDragDropOp> FNamedElementDragDropOp::New(TWeakPtr<SAttributesTreeView::FTreeItem> InElement)
+	TSharedRef<FNamedElementDragDropOp> FNamedElementDragDropOp::New(const FName InNamedAttribute)
 	{
 		TSharedRef<FNamedElementDragDropOp> Operation = MakeShared<FNamedElementDragDropOp>();
-		Operation->Element = InElement;
+		Operation->NamedAttribute = InNamedAttribute;
 		Operation->Construct();
 		return Operation;
 	}
@@ -39,7 +39,7 @@ namespace UE::Anim::STF
 			.BorderImage(FAppStyle::GetBrush("Menu.Background"))
 			[
 				SNew(STextBlock)
-				.Text(FText::FromName(Element.Pin()->Name))
+				.Text(FText::FromName(NamedAttribute))
 			];
 	}
 
@@ -48,7 +48,7 @@ namespace UE::Anim::STF
 	public:
 		DECLARE_DELEGATE_RetVal_TwoParams(bool /* Success */, FOnRenamed, FName /* Old Name */, FName /* New Name */);
 		DECLARE_DELEGATE_TwoParams(FOnReparented, FName /* Attribute Name */, FName /* New Parent Name */);
-		DECLARE_DELEGATE_RetVal_TwoParams(bool /* Success */, FOnTypeSelected, FName /* Attribute Name */, TObjectPtr<const UScriptStruct> /* New Type */);
+		DECLARE_DELEGATE_RetVal_TwoParams(bool /* Success */, FOnTypeSelected, FName /* Attribute Name */, TObjectPtr<const UClass> /* New Type */);
 	
 		SLATE_BEGIN_ARGS(SNamedElementTableRow) { }
 			SLATE_ARGUMENT(TSharedPtr<SAttributesTreeView::FTreeItem>, NamedElementTreeItem)
@@ -61,7 +61,7 @@ namespace UE::Anim::STF
 		{
 			if (MouseEvent.IsMouseButtonDown(EKeys::LeftMouseButton))
 			{
-				const TSharedRef<FNamedElementDragDropOp> DragDropOp = FNamedElementDragDropOp::New(NamedElementTreeItem);
+				const TSharedRef<FNamedElementDragDropOp> DragDropOp = FNamedElementDragDropOp::New(NamedElementTreeItem->Name);
 				return FReply::Handled().BeginDragDrop(DragDropOp);
 			}
 
@@ -89,7 +89,7 @@ namespace UE::Anim::STF
 				return FReply::Unhandled();
 			}
 
-			OnReparented.ExecuteIfBound(DragDropOp->Element.Pin()->Name, NamedElementTreeItem->Name);
+			OnReparented.ExecuteIfBound(DragDropOp->NamedAttribute, NamedElementTreeItem->Name);
 
 			return FReply::Handled();
 		}
@@ -124,14 +124,14 @@ namespace UE::Anim::STF
 			}
 		}
 
-		void OnStructPicked(const UScriptStruct* InStruct)
+		void OnClassPicked(UClass* InClass)
 		{
 			if (OnRenamed.IsBound())
 			{
-				const bool bSuccess = OnTypeSelected.Execute(NamedElementTreeItem->Name, InStruct);
+				const bool bSuccess = OnTypeSelected.Execute(NamedElementTreeItem->Name, InClass);
 				if (bSuccess)
 				{
-					NamedElementTreeItem->Type = InStruct ? (InStruct->GetFName()) : NAME_None;
+					NamedElementTreeItem->Type = InClass ? (InClass->GetFName()) : NAME_None;
 				}
 			}
 			TypeComboButton->SetIsOpen(false);
@@ -139,21 +139,45 @@ namespace UE::Anim::STF
 
 		TSharedRef<SWidget> GenerateStructPicker()
 		{
-			FStructViewerInitializationOptions Options;
-			Options.bShowUnloadedStructs = true;
+			class FClassFilter : public IClassViewerFilter
+			{
+			public:
+				FClassFilter(const UClass* InBaseType) : BaseType(InBaseType)  { };
+		
+				virtual bool IsClassAllowed(const FClassViewerInitializationOptions& InInitOptions, const UClass* InClass, TSharedRef< class FClassViewerFilterFuncs > InFilterFuncs) override
+				{
+					if (InClass->HasMetaData(TEXT("Hidden")))
+					{
+						return false;
+					}
+			
+					return InClass->IsChildOf(BaseType) && InClass != BaseType;
+				}
+				
+				virtual bool IsUnloadedClassAllowed(const FClassViewerInitializationOptions& InInitOptions, const TSharedRef< const class IUnloadedBlueprintData > InUnloadedClassData, TSharedRef<FClassViewerFilterFuncs> InFilterFuncs)
+				{
+					return false;
+				};
+				
+			private:
+				const UClass* BaseType;
+			};
+			
+			FClassViewerInitializationOptions Options;
 			Options.bShowNoneOption = true;
-			Options.NameTypeToDisplay = EStructViewerNameTypeToDisplay::DisplayName;
-			Options.DisplayMode = EStructViewerDisplayMode::ListView;
+			Options.NameTypeToDisplay = EClassViewerNameTypeToDisplay::DisplayName;
+			Options.DisplayMode = EClassViewerDisplayMode::ListView;
+			Options.ClassFilters.Add(MakeShared<FClassFilter>(USkeletonTemplateBaseType::StaticClass()));
 
-			FOnStructPicked OnPicked(FOnStructPicked::CreateRaw(this, &SNamedElementTableRow::OnStructPicked));
+			FOnClassPicked OnPicked(FOnClassPicked::CreateRaw(this, &SNamedElementTableRow::OnClassPicked));
 
-			return FModuleManager::LoadModuleChecked<FStructViewerModule>("StructViewer")
-				.CreateStructViewer(Options, OnPicked);
+			return FModuleManager::LoadModuleChecked<FClassViewerModule>("ClassViewer")
+				.CreateClassViewer(Options, OnPicked);
 		}
 
 		TSharedRef<SWidget> GenerateWidgetForColumn(const FName& ColumnName) override
 		{
-			if (ColumnName == UE::Anim::STF::FOutlinerColumns::NameId)
+			if (ColumnName == FTemplateNamedAttributesColumns::NameId)
 			{
 				TSharedPtr<SInlineEditableTextBlock> InlineWidget;
 
@@ -179,7 +203,8 @@ namespace UE::Anim::STF
 
 				return HorizontalBox;
 			}
-			else if (ColumnName == UE::Anim::STF::FOutlinerColumns::TypeId)
+
+			if (ColumnName == FTemplateNamedAttributesColumns::TypeId)
 			{
 				return SAssignNew(TypeComboButton, SComboButton)
 					.OnGetMenuContent(this, &SNamedElementTableRow::GenerateStructPicker)
@@ -192,10 +217,8 @@ namespace UE::Anim::STF
 							})
 					];
 			}
-			else
-			{
-				return SNullWidget::NullWidget;
-			}
+
+			return SNullWidget::NullWidget;
 		}
 
 	private:
@@ -222,10 +245,10 @@ namespace UE::Anim::STF
 				.HeaderRow
 				(
 					SNew(SHeaderRow)
-						+SHeaderRow::Column(UE::Anim::STF::FOutlinerColumns::NameId)
+						+SHeaderRow::Column(FTemplateNamedAttributesColumns::NameId)
 							.FillWidth(0.5f)
 							.DefaultLabel(LOCTEXT("NameLabel", "Name"))
-						+SHeaderRow::Column(UE::Anim::STF::FOutlinerColumns::TypeId)
+						+SHeaderRow::Column(FTemplateNamedAttributesColumns::TypeId)
 							.FillWidth(0.5f)
 							.DefaultLabel(LOCTEXT("TypeLabel", "Type"))
 				)
@@ -266,7 +289,7 @@ namespace UE::Anim::STF
 					SkeletonTemplate->ReparentNamedAttribute(AttributeName, NewParentName);
 					RegenerateTreeViewItems();
 				})
-			.OnTypeSelected_Lambda([this](const FName AttributeName, const TObjectPtr<const UScriptStruct> NewType)
+			.OnTypeSelected_Lambda([this](const FName AttributeName, const TObjectPtr<const UClass> NewType)
 				{
 					return SkeletonTemplate->SetNamedAttributeType(AttributeName, NewType);
 				})
diff --git a/Engine/Plugins/Experimental/Animation/SkeletonTemplateFramework/Source/Editor/Private/SkeletonTemplateFramework/SkeletonTemplateNamedAttributesEditor.h b/Engine/Plugins/Experimental/Animation/SkeletonTemplateFramework/Source/Editor/Private/SkeletonTemplateFramework/SkeletonTemplateNamedAttributesEditor.h
--- a/Engine/Plugins/Experimental/Animation/SkeletonTemplateFramework/Source/Editor/Private/SkeletonTemplateFramework/SkeletonTemplateNamedAttributesEditor.h
+++ b/Engine/Plugins/Experimental/Animation/SkeletonTemplateFramework/Source/Editor/Private/SkeletonTemplateFramework/SkeletonTemplateNamedAttributesEditor.h
@@ -66,10 +66,10 @@ namespace UE::Anim::STF
 	public:
 		DRAG_DROP_OPERATOR_TYPE(FNamedElementDragDropOp, FDecoratedDragDropOp)
 
-		static TSharedRef<FNamedElementDragDropOp> New(TWeakPtr<SAttributesTreeView::FTreeItem> InElement);
+		static TSharedRef<FNamedElementDragDropOp> New(const FName InNamedAttribute);
 
 		TSharedPtr<SWidget> GetDefaultDecorator() const override;
 
-		TWeakPtr<SAttributesTreeView::FTreeItem> Element;
+		FName NamedAttribute;
 	};
 }
\ No newline at end of file
diff --git a/Engine/Plugins/Experimental/Animation/SkeletonTemplateFramework/Source/Runtime/Private/SkeletonTemplateFramework/SkeletonTemplate.cpp b/Engine/Plugins/Experimental/Animation/SkeletonTemplateFramework/Source/Runtime/Private/SkeletonTemplateFramework/SkeletonTemplate.cpp
--- a/Engine/Plugins/Experimental/Animation/SkeletonTemplateFramework/Source/Runtime/Private/SkeletonTemplateFramework/SkeletonTemplate.cpp
+++ b/Engine/Plugins/Experimental/Animation/SkeletonTemplateFramework/Source/Runtime/Private/SkeletonTemplateFramework/SkeletonTemplate.cpp
@@ -21,7 +21,7 @@ bool USkeletonTemplate::AddNamedAttribute(const FSkeletonNamedAttribute& NewName
 		return false; // Reserved name for denoting unrooted attributes
 	}
 
-	int32 ExistingAttributeIndex = FindNamedAttributeIndex(NewNamedAttribute.Name);
+	const int32 ExistingAttributeIndex = FindNamedAttributeIndex(NewNamedAttribute.Name);
 	if (ExistingAttributeIndex == INDEX_NONE)
 	{
 		NamedAttributes.Add(NewNamedAttribute);
@@ -143,7 +143,7 @@ bool USkeletonTemplate::RemoveNamedAttribute(const FName AttributeToRemoveName,
 	return true;
 }
 
-bool USkeletonTemplate::SetNamedAttributeType(const FName AttributeName, const TObjectPtr<const UScriptStruct> NewType)
+bool USkeletonTemplate::SetNamedAttributeType(const FName AttributeName, const TObjectPtr<const UClass> NewType)
 {
 	FSkeletonNamedAttribute* NamedAttribute = FindMutableNamedAttribute(AttributeName);
 
@@ -161,6 +161,14 @@ const TArray<FSkeletonNamedAttributeSet>& USkeletonTemplate::GetNamedAttributeSe
 	return NamedAttributeSets;
 }
 
+const FSkeletonNamedAttributeSet* const USkeletonTemplate::FindNamedAttributeSet(const FName AttributeSetName) const
+{
+	return NamedAttributeSets.FindByPredicate([AttributeSetName](const FSkeletonNamedAttributeSet& Candidate)
+		{
+			return Candidate.Name == AttributeSetName;
+		});
+}
+
 const FSkeletonNamedAttribute* const USkeletonTemplate::FindNamedAttribute(const FName AttributeName) const
 {
 	return NamedAttributes.FindByPredicate([AttributeName](const FSkeletonNamedAttribute& Candidate)
@@ -222,6 +230,14 @@ bool USkeletonTemplate::RemoveAttributeFromSet(const FName AttributeSetName, con
 	return RemoveCount > 0;
 }
 
+FSkeletonNamedAttributeMapping* USkeletonTemplate::FindMutableNamedAttributeMapping(const FName AttributeMappingName)
+{
+	return NamedAttributeMappings.FindByPredicate([AttributeMappingName](const FSkeletonNamedAttributeMapping& Candidate)
+		{
+			return Candidate.Name == AttributeMappingName;
+		});
+}
+
 FSkeletonNamedAttributeSet* USkeletonTemplate::FindMutableNamedAttributeSet(const FName AttributeSetName)
 {
 	return NamedAttributeSets.FindByPredicate([AttributeSetName](const FSkeletonNamedAttributeSet& Candidate)
@@ -253,6 +269,107 @@ const TArray<FSkeletonNamedAttributeMapping>& USkeletonTemplate::GetNamedAttribu
 	return NamedAttributeMappings;
 }
 
+const FSkeletonNamedAttributeMapping* const USkeletonTemplate::FindNamedAttributeMapping(const FName AttributeMappingName) const
+{
+	return NamedAttributeMappings.FindByPredicate([AttributeMappingName](const FSkeletonNamedAttributeMapping& Candidate)
+		{
+			return Candidate.Name == AttributeMappingName;
+		});
+}
+
+bool USkeletonTemplate::AddAttributeMapping(const FName AttributeMappingName)
+{
+	if (const FSkeletonNamedAttributeMapping* const AttributeMapping = FindNamedAttributeMapping(AttributeMappingName))
+	{
+		return false; // Set already exists with that name
+	}
+
+	FSkeletonNamedAttributeMapping NewMapping;
+	NewMapping.Name = AttributeMappingName;
+
+	NamedAttributeMappings.Add(NewMapping);
+	return true;	
+}
+
+bool USkeletonTemplate::RemoveAttributeMapping(const FName AttributeMappingName)
+{
+	const int32 AttributeMappingIndex = FindNamedAttributeMappingIndex(AttributeMappingName);
+	if (AttributeMappingIndex == INDEX_NONE)
+	{
+		return false; // Cannot find set with that name
+	}
+
+	NamedAttributeSets.RemoveAt(AttributeMappingIndex);
+	return true;
+}
+
+bool USkeletonTemplate::RenameNamedAttributeMapping(const FName AttributeMappingName, const FName NewAttributeMappingName)
+{
+	FSkeletonNamedAttributeMapping* AttributeMappingToRename = FindMutableNamedAttributeMapping(AttributeMappingName);
+	if (!AttributeMappingToRename)
+	{
+		return false; // No attribute mapping with old name
+	}
+
+	const FSkeletonNamedAttributeMapping* const ExistingAttributeMapping = FindNamedAttributeMapping(NewAttributeMappingName);
+	if (ExistingAttributeMapping)
+	{
+		return false; // Attribute mapping with new name already exists
+	}
+
+	AttributeMappingToRename->Name = NewAttributeMappingName;
+	return true;
+}
+
+bool USkeletonTemplate::SetNamedAttributeMappingSourceSet(const FName AttributeMappingName, const FName NewSourceSet)
+{
+	FSkeletonNamedAttributeMapping* AttributeMappingToUpdate = FindMutableNamedAttributeMapping(AttributeMappingName);
+	if (!AttributeMappingToUpdate)
+	{
+		return false; // No attribute mapping found
+	}
+
+	if (NewSourceSet != NAME_None)
+	{
+		const FSkeletonNamedAttributeSet* const AttributeSet = FindNamedAttributeSet(NewSourceSet);
+		if (!AttributeSet)
+		{
+			return false; // No attribute set found
+		}
+
+		AttributeMappingToUpdate->SourceAttributeSet = NewSourceSet;
+		AttributeMappingToUpdate->TableData.Reset(); // TODO: Keep table entries for attributes from the old source set that also exist in the new set
+
+		for (const FName& AttributeName : AttributeSet->NamedAttributes)
+		{
+			FSkeletonNamedAttributeMappingEntry NewMappingEntry;
+			NewMappingEntry.AttributeName = AttributeName;
+
+			if (const FSkeletonNamedAttribute* const NamedAttribute = FindNamedAttribute(AttributeName))
+			{
+				const FSkeletonNamedAttributeMappingType* const MappingType = AttributeMappingToUpdate->MappingTypes.FindByPredicate([NamedAttribute](const FSkeletonNamedAttributeMappingType& MappingType)
+				{
+					return MappingType.SourceType == NamedAttribute->Type;
+				});
+
+				if (MappingType)
+				{
+					NewMappingEntry.Value = NewObject<USkeletonTemplateBaseType>(this, MappingType->TargetType);
+				}
+			}
+			
+			AttributeMappingToUpdate->TableData.Add(NewMappingEntry);
+		}
+	}
+	else
+	{
+		AttributeMappingToUpdate->SourceAttributeSet = NAME_None;
+		AttributeMappingToUpdate->TableData.Reset(); // TODO: Keep table entries for attributes from the old source set that also exist in the new set
+	}
+
+	return true;
+}
+
 FSkeletonNamedAttribute* USkeletonTemplate::FindMutableNamedAttribute(const FName AttributeName)
 {
 	return NamedAttributes.FindByPredicate([AttributeName](const FSkeletonNamedAttribute& Candidate)
@@ -269,6 +386,14 @@ int32 USkeletonTemplate::FindNamedAttributeSetIndex(const FName AttributeSetName
 		});
 }
 
+int32 USkeletonTemplate::FindNamedAttributeMappingIndex(const FName AttributeMappingName)
+{
+	return NamedAttributeMappings.IndexOfByPredicate([AttributeMappingName](const FSkeletonNamedAttributeMapping& Candidate)
+		{
+			return Candidate.Name == AttributeMappingName;
+		});
+}
+
 int32 USkeletonTemplate::FindNamedAttributeIndex(const FName AttributeName)
 {
 	return NamedAttributes.IndexOfByPredicate([AttributeName](const FSkeletonNamedAttribute& Candidate)
@@ -277,7 +402,10 @@ int32 USkeletonTemplate::FindNamedAttributeIndex(const FName AttributeName)
 		});
 }
 
-void USkeletonTemplate::SortNamedAttributes()
+// Util function to sort an array topologically for any type with unique FName identifiers and parent identifiers
+// TODO: Will be used if needed to sort entries in sets and mappings similarly to how named mappings are sorted
+template <typename ArrayType>
+void SortByProperty(TArray<ArrayType>& Array, FName ArrayType::*NameProperty, FName ArrayType::*ParentNameProperty, TFunctionRef<ArrayType*(FName)> FindArrayItem)
 {
 	struct Node
 	{
@@ -294,16 +422,16 @@ void USkeletonTemplate::SortNamedAttributes()
 
 	TMap<FName, Node> Graph;
 
-	for (const FSkeletonNamedAttribute& NamedAttribute : NamedAttributes)
+	for (const ArrayType& ArrayItem : Array)
 	{
-		Graph.Add(NamedAttribute.Name, Node(NamedAttribute.Name));
-		Graph.Add(NamedAttribute.ParentName, Node(NamedAttribute.ParentName));
+		Graph.Add(ArrayItem.*NameProperty, Node(ArrayItem.*NameProperty));
+		Graph.Add(ArrayItem.*ParentNameProperty, Node(ArrayItem.*ParentNameProperty));
 	}
 
-	for (const FSkeletonNamedAttribute& NamedAttribute : NamedAttributes)
+	for (const ArrayType& ArrayItem : Array)
 	{
-		Graph[NamedAttribute.ParentName].Children.Add(NamedAttribute.Name);
-		++Graph[NamedAttribute.Name].InDegree;
+		Graph[ArrayItem.*ParentNameProperty].Children.Add(ArrayItem.*NameProperty);
+		++Graph[ArrayItem.*NameProperty].InDegree;
 	}
 
 	TArray<FName> NodeQueue;
@@ -315,14 +443,14 @@ void USkeletonTemplate::SortNamedAttributes()
 		}
 	}
 
-	TArray<FSkeletonNamedAttribute> SortedNamedAttributes;
+	TArray<ArrayType> SortedArray;
 
 	while (!NodeQueue.IsEmpty())
 	{
 		const FName AttributeName = NodeQueue.Pop();
 		if (AttributeName != NAME_None)
 		{
-			SortedNamedAttributes.Add(*FindNamedAttribute(AttributeName));
+			SortedArray.Add(*FindArrayItem(AttributeName));
 		}
 
 		for (const FName& Child : Graph[AttributeName].Children)
@@ -335,6 +463,16 @@ void USkeletonTemplate::SortNamedAttributes()
 		}
 	}
 
-	ensureMsgf(NamedAttributes.Num() == SortedNamedAttributes.Num(), TEXT("Cycle found in the named attribute DAG"));
-	NamedAttributes = SortedNamedAttributes;
+	ensureMsgf(Array.Num() == SortedArray.Num(), TEXT("Cycle found in the named attribute DAG"));
+	Array = SortedArray;
+}
+
+void USkeletonTemplate::SortNamedAttributes()
+{
+	TFunction<FSkeletonNamedAttribute*(FName)> Functor = [this](FName InAttributeName)
+	{
+		return FindMutableNamedAttribute(InAttributeName);
+	};
+
+	SortByProperty<FSkeletonNamedAttribute>(NamedAttributes, &FSkeletonNamedAttribute::Name, &FSkeletonNamedAttribute::ParentName, Functor);
 }
\ No newline at end of file
diff --git a/Engine/Plugins/Experimental/Animation/SkeletonTemplateFramework/Source/Runtime/Public/SkeletonTemplateFramework/NamedElementTypes.h b/Engine/Plugins/Experimental/Animation/SkeletonTemplateFramework/Source/Runtime/Public/SkeletonTemplateFramework/NamedElementTypes.h
--- a/Engine/Plugins/Experimental/Animation/SkeletonTemplateFramework/Source/Runtime/Public/SkeletonTemplateFramework/NamedElementTypes.h
+++ b/Engine/Plugins/Experimental/Animation/SkeletonTemplateFramework/Source/Runtime/Public/SkeletonTemplateFramework/NamedElementTypes.h
@@ -3,14 +3,48 @@
 #pragma once
 
 #include "CoreMinimal.h"
+#include "Animation/Skeleton.h"
 
 #include "NamedElementTypes.generated.h"
 
-USTRUCT()
-struct FSkeletonNamedAttributeType_Transform
+UCLASS(MinimalAPI)
+class USkeletonTemplateBaseType : public UObject
 {
 	GENERATED_BODY()
+};
 
-	UPROPERTY()
+UCLASS(MinimalAPI)
+class USkeletonTemplateTransform : public USkeletonTemplateBaseType
+{
+	GENERATED_BODY()
+
+	UPROPERTY(EditAnywhere, Category = "Settings")
 	FTransform Value;
+};
+
+UCLASS(MinimalAPI)
+class USkeletonTemplateFloat : public USkeletonTemplateBaseType
+{
+	GENERATED_BODY()
+
+	UPROPERTY(EditAnywhere, Category = "Settings")
+	float Value;
+};
+
+UCLASS(MinimalAPI)
+class USkeletonTemplateBool : public USkeletonTemplateBaseType
+{
+	GENERATED_BODY()
+
+	UPROPERTY(EditAnywhere, Category = "Settings")
+	bool Value;
+};
+
+UCLASS(MinimalAPI)
+class USkeletonTemplateTranslationRetargetingMode : public USkeletonTemplateBaseType
+{
+	GENERATED_BODY()
+
+	UPROPERTY(EditAnywhere, Category = "Settings")
+	TEnumAsByte<EBoneTranslationRetargetingMode::Type> Value;
 };
\ No newline at end of file
diff --git a/Engine/Plugins/Experimental/Animation/SkeletonTemplateFramework/Source/Runtime/Public/SkeletonTemplateFramework/SkeletonTemplate.h b/Engine/Plugins/Experimental/Animation/SkeletonTemplateFramework/Source/Runtime/Public/SkeletonTemplateFramework/SkeletonTemplate.h
--- a/Engine/Plugins/Experimental/Animation/SkeletonTemplateFramework/Source/Runtime/Public/SkeletonTemplateFramework/SkeletonTemplate.h
+++ b/Engine/Plugins/Experimental/Animation/SkeletonTemplateFramework/Source/Runtime/Public/SkeletonTemplateFramework/SkeletonTemplate.h
@@ -3,7 +3,7 @@
 #pragma once
 
 #include "CoreMinimal.h"
-#include "HierarchyTable.h"
+#include "SkeletonTemplateFramework/NamedElementTypes.h"
 
 #include "SkeletonTemplate.generated.h"
 
@@ -23,7 +23,7 @@ struct FSkeletonNamedAttribute
 
 	// The type of the attribute
 	UPROPERTY(EditAnywhere, Category = "Settings")
-	TObjectPtr<const UScriptStruct> Type;
+	TObjectPtr<const UClass> Type;
 };
 
 // A skeleton named attribute set
@@ -44,6 +44,36 @@ struct FSkeletonNamedAttributeSet
 	// Sub-parts are defined per attribute type (e.g. FTransform has rotation/translation/scale sub-parts)
 };
 
+// A struct stored for each entry in a named attribute mapping
+USTRUCT()
+struct FSkeletonNamedAttributeMappingEntry
+{
+	GENERATED_BODY()
+
+	// The attribute in the set that this value relates to
+	UPROPERTY(EditAnywhere, Category = "Settings")
+	FName AttributeName;
+
+	// The actual payload data for this entry
+	UPROPERTY(EditAnywhere, Category = "Settings")
+	TObjectPtr<USkeletonTemplateBaseType> Value;
+};
+
+// A mapping struct for each named attribute mapping to describe what payload type to store for each unique type of named attribute
+USTRUCT()
+struct FSkeletonNamedAttributeMappingType
+{
+	GENERATED_BODY()
+
+	// The source named attribute type
+	UPROPERTY(EditAnywhere, Category = "Settings")
+	TObjectPtr<const UClass> SourceType;
+
+	// The target mapping payload type
+	UPROPERTY(EditAnywhere, Category = "Settings")
+	TObjectPtr<const UClass> TargetType;
+};
+
 // A skeleton named attribute mapping
 USTRUCT()
 struct FSkeletonNamedAttributeMapping
@@ -58,9 +88,13 @@ struct FSkeletonNamedAttributeMapping
 	UPROPERTY(EditAnywhere, Category = "Settings")
 	FName SourceAttributeSet;
 
+	// A mapping struct for each named attribute mapping to describe what payload type to store for each unique type of named attribute
+	UPROPERTY(EditAnywhere, Category = "Settings")
+	TArray<FSkeletonNamedAttributeMappingType> MappingTypes;
+
 	// The hierarchy table that contains the mapping data
 	UPROPERTY(EditAnywhere, Category = "Settings")
-	TObjectPtr<UHierarchyTable> TableData;
+	TArray<FSkeletonNamedAttributeMappingEntry> TableData;
 };
 
 // A skeleton template specifies a list of attributes that skeletons can contain as well
@@ -75,6 +109,8 @@ class USkeletonTemplate : public UObject
 public:
 	USkeletonTemplate(const FObjectInitializer& ObjectInitializer);
 
+	// Attributes
+
 	// Returns topologically sorted attributes, i.e. parents are defined before they're referenced by children
 	SKELETONTEMPLATEFRAMEWORKRUNTIME_API const TArray<FSkeletonNamedAttribute>& GetNamedAttributes() const;
 
@@ -92,11 +128,15 @@ public:
 
 	SKELETONTEMPLATEFRAMEWORKRUNTIME_API bool RemoveNamedAttribute(const FName AttributeName, const ERemoveNamedAttributeBehavior Behavior = ERemoveNamedAttributeBehavior::MoveChildrenToParent);
 
-	SKELETONTEMPLATEFRAMEWORKRUNTIME_API bool SetNamedAttributeType(const FName AttributeName, const TObjectPtr<const UScriptStruct> NewType);
+	SKELETONTEMPLATEFRAMEWORKRUNTIME_API bool SetNamedAttributeType(const FName AttributeName, const TObjectPtr<const UClass> NewType);
+
+	SKELETONTEMPLATEFRAMEWORKRUNTIME_API const FSkeletonNamedAttribute* const FindNamedAttribute(const FName AttributeName) const;
 
+	// Sets
+	
 	SKELETONTEMPLATEFRAMEWORKRUNTIME_API const TArray<FSkeletonNamedAttributeSet>& GetNamedAttributeSets() const;
 
-	SKELETONTEMPLATEFRAMEWORKRUNTIME_API const FSkeletonNamedAttribute* const FindNamedAttribute(const FName AttributeName) const;
+	SKELETONTEMPLATEFRAMEWORKRUNTIME_API const FSkeletonNamedAttributeSet* const FindNamedAttributeSet(const FName AttributeSetName) const;
 
 	SKELETONTEMPLATEFRAMEWORKRUNTIME_API bool AddAttributeSet(const FName AttributeSetName);
 
@@ -108,17 +148,33 @@ public:
 
 	SKELETONTEMPLATEFRAMEWORKRUNTIME_API bool RenameNamedAttributeSet(const FName AttributeSetName, const FName NewAttributeSetName);
 
+	// Mappings
+
 	SKELETONTEMPLATEFRAMEWORKRUNTIME_API const TArray<FSkeletonNamedAttributeMapping>& GetNamedAttributeMappings() const;
 
+	SKELETONTEMPLATEFRAMEWORKRUNTIME_API const FSkeletonNamedAttributeMapping* const FindNamedAttributeMapping(const FName AttributeMappingName) const;
+
+	SKELETONTEMPLATEFRAMEWORKRUNTIME_API bool AddAttributeMapping(const FName AttributeMappingName);
+
+	SKELETONTEMPLATEFRAMEWORKRUNTIME_API bool RemoveAttributeMapping(const FName AttributeMappingName);
+
+	SKELETONTEMPLATEFRAMEWORKRUNTIME_API bool RenameNamedAttributeMapping(const FName AttributeMappingName, const FName NewAttributeMappingName);
+
+	SKELETONTEMPLATEFRAMEWORKRUNTIME_API bool SetNamedAttributeMappingSourceSet(const FName AttributeMappingName, const FName NewSourceSet);
+
 protected:
-	FSkeletonNamedAttribute* FindMutableNamedAttribute(const FName AttributeName);
+	FSkeletonNamedAttributeMapping* FindMutableNamedAttributeMapping(const FName AttributeMappingName);
 
 	FSkeletonNamedAttributeSet* FindMutableNamedAttributeSet(const FName AttributeSetName);
 
+	FSkeletonNamedAttribute* FindMutableNamedAttribute(const FName AttributeName);
+
+	int32 FindNamedAttributeIndex(const FName AttributeName);
+	
 	int32 FindNamedAttributeSetIndex(const FName AttributeSetName);
 	
-	int32 FindNamedAttributeIndex(const FName AttributeName);
-
+	int32 FindNamedAttributeMappingIndex(const FName AttributeMappingName);
+	
 	void SortNamedAttributes();
 
 	// TODO Parent Template Support:
