Skeleton Binding Editor

diff --git a/Engine/Plugins/Experimental/Animation/SkeletonTemplateFramework/Source/Editor/Private/SkeletonTemplateFramework/SkeletonBindingEditorToolkit.cpp b/Engine/Plugins/Experimental/Animation/SkeletonTemplateFramework/Source/Editor/Private/SkeletonTemplateFramework/SkeletonBindingEditorToolkit.cpp
--- a/Engine/Plugins/Experimental/Animation/SkeletonTemplateFramework/Source/Editor/Private/SkeletonTemplateFramework/SkeletonBindingEditorToolkit.cpp
+++ b/Engine/Plugins/Experimental/Animation/SkeletonTemplateFramework/Source/Editor/Private/SkeletonTemplateFramework/SkeletonBindingEditorToolkit.cpp
@@ -5,6 +5,9 @@
 #include "Modules/ModuleManager.h"
 #include "PropertyEditorModule.h"
 #include "SkeletonTemplateFramework/SkeletonBinding.h"
+#include "SkeletonTemplateFramework/SkeletonBindingNamedAttributesEditor.h"
+#include "SkeletonTemplateFramework/SkeletonBindingNamedAttributeSetsEditor.h"
+#include "SkeletonTemplateFramework/SkeletonBindingNamedAttributeMappingsEditor.h"
 #include "Widgets/Docking/SDockTab.h"
 
 #define LOCTEXT_NAMESPACE "SkeletonBindingEditorToolkit"
@@ -137,19 +140,40 @@ FLinearColor FSkeletonBindingEditorToolkit::GetWorldCentricTabColorScale() const
 	return FLinearColor::White;
 }
 
+void FSkeletonBindingEditorToolkit::SetDetailsObject(TObjectPtr<UObject> InObject)
+{
+	if (InObject)
+	{
+		DetailsView->SetObjects(TArray{ InObject });
+	}
+	else
+	{
+		DetailsView->SetObjects(TArray<TObjectPtr<UObject>>{ SkeletonBinding });
+	}
+}
+
 TSharedRef<SDockTab> FSkeletonBindingEditorToolkit::SpawnTabAttributes(const FSpawnTabArgs& Args)
 {
-	return SNew(SDockTab);
+	return SNew(SDockTab)
+		[
+			SAssignNew(AttributeBindingsTreeView, UE::Anim::STF::SAttributeBindingsTreeView, SkeletonBinding)
+		];
 }
 
 TSharedRef<SDockTab> FSkeletonBindingEditorToolkit::SpawnTabAttributeSets(const FSpawnTabArgs& Args)
 {
-	return SNew(SDockTab);
+	return SNew(SDockTab)
+		[
+			SAssignNew(BindingSetsTreeView, UE::Anim::STF::SBindingSetsTreeView, SkeletonBinding)
+		];
 }
 
 TSharedRef<SDockTab> FSkeletonBindingEditorToolkit::SpawnTabAttributeMappings(const FSpawnTabArgs& Args)
 {
-	return SNew(SDockTab);
+	return SNew(SDockTab)
+		[
+			SAssignNew(BindingMappingsTreeView, UE::Anim::STF::SBindingMappingsTreeView, SkeletonBinding, this)
+		];
 }
 
 TSharedRef<SDockTab> FSkeletonBindingEditorToolkit::SpawnTabDetails(const FSpawnTabArgs& Args)
diff --git a/Engine/Plugins/Experimental/Animation/SkeletonTemplateFramework/Source/Editor/Private/SkeletonTemplateFramework/SkeletonBindingEditorToolkit.h b/Engine/Plugins/Experimental/Animation/SkeletonTemplateFramework/Source/Editor/Private/SkeletonTemplateFramework/SkeletonBindingEditorToolkit.h
--- a/Engine/Plugins/Experimental/Animation/SkeletonTemplateFramework/Source/Editor/Private/SkeletonTemplateFramework/SkeletonBindingEditorToolkit.h
+++ b/Engine/Plugins/Experimental/Animation/SkeletonTemplateFramework/Source/Editor/Private/SkeletonTemplateFramework/SkeletonBindingEditorToolkit.h
@@ -9,12 +9,20 @@ class USkeletonBinding;
 class SDockTab;
 namespace UE::Anim::STF
 {
-	class SAttributesTreeView;
-	class SAttributeSetsTreeView;
-	class SAttributeMappingsTreeView;
+	class SAttributeBindingsTreeView;
+	class SBindingSetsTreeView;
+	class SBindingMappingsTreeView;
 }
 
-class FSkeletonBindingEditorToolkit : public FAssetEditorToolkit
+class ISkeletonBindingEditorToolkit
+{
+public:
+	virtual ~ISkeletonBindingEditorToolkit() = default;
+	
+	virtual void SetDetailsObject(TObjectPtr<UObject> InObject) = 0;
+};
+
+class FSkeletonBindingEditorToolkit : public FAssetEditorToolkit, public ISkeletonBindingEditorToolkit
 {
 public:
 	void InitEditor(const TArray<UObject*>& InObjects);
@@ -27,6 +35,9 @@ public:
 	FString GetWorldCentricTabPrefix() const override;
 	FLinearColor GetWorldCentricTabColorScale() const override;
 
+	// ISkeletonTemplateEditorToolkit
+	void SetDetailsObject(TObjectPtr<UObject> InObject) override;
+
 private:
 	TSharedRef<SDockTab> SpawnTabAttributes(const FSpawnTabArgs& Args);
 	TSharedRef<SDockTab> SpawnTabAttributeSets(const FSpawnTabArgs& Args);
@@ -36,4 +47,10 @@ private:
 	TObjectPtr<USkeletonBinding> SkeletonBinding;
 
 	TSharedPtr<IDetailsView> DetailsView;
+
+	TSharedPtr<UE::Anim::STF::SAttributeBindingsTreeView> AttributeBindingsTreeView;
+
+	TSharedPtr<UE::Anim::STF::SBindingSetsTreeView> BindingSetsTreeView;
+
+	TSharedPtr<UE::Anim::STF::SBindingMappingsTreeView> BindingMappingsTreeView;
 };
\ No newline at end of file
diff --git a/Engine/Plugins/Experimental/Animation/SkeletonTemplateFramework/Source/Editor/Private/SkeletonTemplateFramework/SkeletonBindingNamedAttributeMappingsEditor.cpp b/Engine/Plugins/Experimental/Animation/SkeletonTemplateFramework/Source/Editor/Private/SkeletonTemplateFramework/SkeletonBindingNamedAttributeMappingsEditor.cpp
new file mode 100644
--- /dev/null
+++ b/Engine/Plugins/Experimental/Animation/SkeletonTemplateFramework/Source/Editor/Private/SkeletonTemplateFramework/SkeletonBindingNamedAttributeMappingsEditor.cpp
@@ -0,0 +1,458 @@
+// Copyright Epic Games, Inc. All Rights Reserved.
+
+#include "SkeletonBindingNamedAttributeMappingsEditor.h"
+
+#include "Framework/MultiBox/MultiBoxBuilder.h"
+#include "Modules/ModuleManager.h"
+#include "SkeletonTemplateFramework/SkeletonBinding.h"
+#include "Widgets/Input/SCheckBox.h"
+#include "Widgets/Input/SComboBox.h"
+#include "Widgets/Text/SInlineEditableTextBlock.h"
+#include "Widgets/Views/STreeView.h"
+#include "Animation/Skeleton.h"
+#include "PropertyEditorModule.h"
+#include "SkeletonBindingEditorToolkit.h"
+
+#define LOCTEXT_NAMESPACE "SkeletonBindingNamedAttributeMappingsEditor"
+
+namespace UE::Anim::STF
+{
+	struct FBindingAttributeTreeItem;
+
+	struct FBindingAttributeTreeItem : public SBindingMappingsTreeView::ITreeItem
+	{
+		TObjectPtr<const UClass> ValueType;
+		bool bHasValue;
+		TSharedPtr<FBindingAttributeTreeItem> Parent;
+		TArray<TSharedPtr<FBindingAttributeTreeItem>> ChildrenAttributes;
+
+		virtual void GetChildren(TArray<TSharedPtr<ITreeItem>>& OutChildren) const override
+		{
+			OutChildren.Append(ChildrenAttributes);
+		}
+	};
+	
+	struct FBindingNamedAttributeMappingListColumns
+	{
+		static const FName NameId;
+		static const FName SourceSetId;
+	};
+
+	const FName FBindingNamedAttributeMappingListColumns::NameId(TEXT("Name"));
+	const FName FBindingNamedAttributeMappingListColumns::SourceSetId(TEXT("SourceSet"));
+	
+	struct FBindingNamedAttributeMappingsColumns
+	{
+		static const FName NameId;
+		static const FName TypeId;
+		static const FName ValueId;
+	};
+
+	const FName FBindingNamedAttributeMappingsColumns::NameId(TEXT("Name"));
+	const FName FBindingNamedAttributeMappingsColumns::TypeId(TEXT("Type"));
+	const FName FBindingNamedAttributeMappingsColumns::ValueId(TEXT("Value"));
+
+	class SBindingMappingTableRow : public SMultiColumnTableRow<TSharedPtr<SBindingMappingsTreeView::ITreeItem>>
+	{
+	public:
+		DECLARE_DELEGATE_RetVal_TwoParams(bool /* Success */, FOnRenamed, FName /* Old Name */, FName /* New Name */);
+
+		SLATE_BEGIN_ARGS(SBindingMappingTableRow) { }
+			SLATE_ARGUMENT(TSharedPtr<SBindingMappingsTreeView::ITreeItem>, TreeItem)
+			SLATE_EVENT(FOnRenamed, OnRenamed);
+		SLATE_END_ARGS()
+
+		void OnCommitRename(const FText& InText, ETextCommit::Type CommitInfo)
+		{
+			const FName OldName = TreeItem->AttributeName;
+			const FName NewName = FName(InText.ToString());
+
+			if (OnRenamed.IsBound())
+			{
+				const bool bSuccess = OnRenamed.Execute(OldName, NewName);
+				if (bSuccess)
+				{
+					TreeItem->AttributeName = NewName;
+				}
+			}
+		}
+
+		void Construct(const FArguments& InArgs, const TSharedRef<STableViewBase>& InOwnerTableView, const TSharedPtr<SBindingMappingsTreeView> InTreeView)
+		{
+			TreeItem = InArgs._TreeItem;
+			OnRenamed = InArgs._OnRenamed;
+			TreeView = InTreeView;
+
+			FSuperRowType::Construct(
+				FSuperRowType::FArguments(),
+				InOwnerTableView);
+		}
+
+		TSharedRef<SWidget> GenerateWidgetForColumn(const FName& ColumnName) override
+		{
+			if (ColumnName == FBindingNamedAttributeMappingsColumns::NameId)
+			{
+				return SNew(SHorizontalBox)
+					+ SHorizontalBox::Slot()
+						.AutoWidth()
+						[
+							SNew(SExpanderArrow, SharedThis(this))
+								.ShouldDrawWires(true)
+						]
+					+ SHorizontalBox::Slot()
+						.AutoWidth()
+						[
+							SNew(STextBlock)
+								.Text_Lambda([this]()
+									{
+										return FText::FromName(TreeItem->AttributeName);
+									})
+						];
+			}
+				
+			if (ColumnName == FBindingNamedAttributeMappingsColumns::TypeId)
+			{
+				return SNew(SHorizontalBox)
+					+ SHorizontalBox::Slot()
+						.VAlign(VAlign_Center)
+						[
+							SNew(STextBlock)
+								.Text(TreeItem->AttributeType)
+						];
+			}
+				
+			if (ColumnName == FBindingNamedAttributeMappingsColumns::ValueId)
+			{
+				// Shown in details panel, TODO: Show read-only widget here?
+				return SNullWidget::NullWidget;
+			}
+			
+			return SNullWidget::NullWidget;
+			
+		}
+
+	private:
+		TSharedPtr<SBindingMappingsTreeView::ITreeItem> TreeItem;
+		TSharedPtr<SBindingMappingsTreeView> TreeView;
+		FOnRenamed OnRenamed;
+	};
+
+	void SBindingMappingsTreeView::Construct(const FArguments& InArgs, TObjectPtr<USkeletonBinding> InSkeletonBinding, ISkeletonBindingEditorToolkit* InToolkit)
+	{
+		SkeletonBinding = InSkeletonBinding;
+
+		FPropertyEditorModule& PropertyEditorModule = FModuleManager::GetModuleChecked<FPropertyEditorModule>("PropertyEditor");
+
+		Toolkit = InToolkit;
+							
+		ChildSlot
+		[
+			SNew(SVerticalBox)
+				+ SVerticalBox::Slot()
+				.MinHeight(200.0f)
+				.AutoHeight()
+				[
+					SAssignNew(ListView, SListView<TSharedPtr<FListItem>>)
+						.SelectionMode(ESelectionMode::Single)
+						.ListItemsSource(&ListItems)
+						.HeaderRow
+						(
+							SNew(SHeaderRow)
+								+ SHeaderRow::Column(FBindingNamedAttributeMappingListColumns::NameId)
+									.FillWidth(0.5f)
+									.DefaultLabel(LOCTEXT("NameLabel", "Name"))
+								+ SHeaderRow::Column(FBindingNamedAttributeMappingListColumns::SourceSetId)
+									.FillWidth(0.5f)
+									.DefaultLabel(LOCTEXT("SourceSetLabel", "Source Set"))
+						)
+						.OnGenerateRow(this, &SBindingMappingsTreeView::ListView_GenerateItemRow)
+						.OnSelectionChanged(this, &SBindingMappingsTreeView::ListView_OnSelectionChanged)
+						.OnContextMenuOpening(this, &SBindingMappingsTreeView::ListView_HandleContextMenuOpening)
+				]
+				+ SVerticalBox::Slot()
+				.FillHeight(1.0f)
+				[
+					SAssignNew(TreeView, STreeView<TSharedPtr<ITreeItem>>)
+						.SelectionMode(ESelectionMode::Single)
+						.TreeItemsSource(&RootItems)
+						.OnGenerateRow(this, &SBindingMappingsTreeView::TreeView_GenerateItemRow)
+						.OnSelectionChanged(this, &SBindingMappingsTreeView::TreeView_OnSelectionChanged)
+						.OnGetChildren(this, &SBindingMappingsTreeView::TreeView_HandleGetChildren)
+						.OnContextMenuOpening(this, &SBindingMappingsTreeView::TreeView_HandleContextMenuOpening)
+						.HighlightParentNodesForSelection(true)
+						.HeaderRow
+						(
+							SNew(SHeaderRow)
+								+SHeaderRow::Column(FBindingNamedAttributeMappingsColumns::NameId)
+									.FillWidth(0.5f)
+									.DefaultLabel(LOCTEXT("NameLabel", "Name"))
+								/*
+								+SHeaderRow::Column(UE::Anim::STF::FNamedAttributeMappingsColumns::TypeId)
+									.FillWidth(0.1f)
+									.DefaultLabel(LOCTEXT("TypeLabel", "Type"))
+								*/
+								+SHeaderRow::Column(FBindingNamedAttributeMappingsColumns::ValueId)
+									.FillWidth(0.5f)
+									.DefaultLabel(LOCTEXT("ValueLabel", "Value"))
+						)
+				]
+		];
+
+		RegenerateListViewItems();
+		RegenerateTreeViewItems();
+
+		// Expand all tree items on construction
+		for (TSharedPtr<ITreeItem> TreeItem : GetAllTreeItems())
+		{
+			TreeView->SetItemExpansion(TreeItem, true);
+		}
+	}
+	
+	void SBindingMappingsTreeView::OnNamedAttributeSetsChanged()
+	{
+		RegenerateListViewItems();
+	}
+
+	class SBindingListViewRow : public SMultiColumnTableRow<TSharedPtr<SBindingMappingsTreeView::FListItem>>
+	{
+	public:
+		SLATE_BEGIN_ARGS(SBindingListViewRow) { }
+			SLATE_ARGUMENT(TSharedPtr<SBindingMappingsTreeView::FListItem>, ListItem)
+		SLATE_END_ARGS()
+
+		void Construct(const FArguments& InArgs, const TSharedRef<STableViewBase>& InOwnerTableView)
+		{
+			ListItem = InArgs._ListItem;
+
+			FSuperRowType::Construct(
+				FSuperRowType::FArguments(),
+				InOwnerTableView);
+		}
+
+		TSharedRef<SWidget> GenerateWidgetForColumn(const FName& ColumnName) override
+		{
+			if (ColumnName == FBindingNamedAttributeMappingListColumns::NameId)
+			{
+				return SNew(STextBlock)
+					.Text(FText::FromName(ListItem->Name));
+			}
+
+			if (ColumnName == FBindingNamedAttributeMappingListColumns::SourceSetId)
+			{
+				return SNew(STextBlock)
+					.Text(FText::FromName(ListItem->SourceSetName));
+			}
+
+			return SNullWidget::NullWidget;
+		}
+
+	private:
+		TSharedPtr<SBindingMappingsTreeView::FListItem> ListItem;
+	};
+	
+	TSharedRef<ITableRow> SBindingMappingsTreeView::ListView_GenerateItemRow(TSharedPtr<FListItem> Item, const TSharedRef<STableViewBase>& OwnerTable)
+	{
+		return SNew(SBindingListViewRow, OwnerTable)
+			.ListItem(Item);
+	}
+
+	void SBindingMappingsTreeView::ListView_OnSelectionChanged(TSharedPtr<FListItem> InItem, ESelectInfo::Type InSelectInfo)
+	{
+		SelectedMappingName = InItem ? InItem->Name : NAME_None;
+		RegenerateTreeViewItems();
+	}
+	
+	TSharedPtr<SWidget> SBindingMappingsTreeView::ListView_HandleContextMenuOpening()
+	{
+		FMenuBuilder MenuBuilder(true, nullptr);		
+		return MenuBuilder.MakeWidget();
+	}
+
+	void SBindingMappingsTreeView::TreeView_HandleGetChildren(TSharedPtr<ITreeItem> InItem, TArray<TSharedPtr<ITreeItem>>& OutChildren) const
+	{
+		InItem->GetChildren(OutChildren);
+	}
+
+	void SBindingMappingsTreeView::TreeView_OnSelectionChanged(TSharedPtr<ITreeItem> InItem, ESelectInfo::Type InSelectInfo)
+	{
+		if (InItem)
+		{			
+			const FSkeletonNamedAttributeMapping* const Mapping = SkeletonBinding->FindNamedAttributeMapping(SelectedMappingName);
+			const FSkeletonNamedAttributeMappingEntry* const Entry = Mapping->TableData.FindByPredicate([InItem](const FSkeletonNamedAttributeMappingEntry& Candidate)
+			{
+				return Candidate.AttributeName == InItem->AttributeName;
+			});
+
+			if (Entry)
+			{
+				Toolkit->SetDetailsObject(Entry->Value);
+			}
+		}
+		else
+		{
+			Toolkit->SetDetailsObject(nullptr);
+		}
+	}
+
+	TSharedRef<ITableRow> SBindingMappingsTreeView::TreeView_GenerateItemRow(TSharedPtr<SBindingMappingsTreeView::ITreeItem> Item, const TSharedRef<STableViewBase>& OwnerTable)
+	{
+		return SNew(SBindingMappingTableRow, OwnerTable, SharedThis(this))
+			.TreeItem(Item);
+	}
+
+	TSharedPtr<SWidget> SBindingMappingsTreeView::TreeView_HandleContextMenuOpening()
+	{
+		FMenuBuilder MenuBuilder(true, nullptr);
+		return MenuBuilder.MakeWidget();
+	}
+
+	TArray<TSharedPtr<UE::Anim::STF::SBindingMappingsTreeView::ITreeItem>> SBindingMappingsTreeView::GetAllTreeItems()
+	{
+		using namespace UE::Anim::STF;
+
+		TArray<TSharedPtr<SBindingMappingsTreeView::ITreeItem>> AllItems;
+		AllItems.Append(RootItems);
+	
+		for (int32 Index = 0; Index < AllItems.Num(); ++Index)
+		{
+			AllItems[Index]->GetChildren(AllItems);
+		}
+
+		return AllItems;
+	}
+
+	void SBindingMappingsTreeView::RegenerateListViewItems()
+	{
+		ListItems.Reset();
+
+		const TArray<FSkeletonNamedAttributeMapping>& Mappings = SkeletonBinding->GetNamedAttributeMappings();
+		for (const FSkeletonNamedAttributeMapping& Mapping : Mappings)
+		{
+			TSharedPtr<FListItem> ListItem = MakeShared<FListItem>();
+			ListItem->Name = Mapping.Name;
+			ListItem->SourceSetName = Mapping.SourceAttributeSet;
+
+			ListItems.Add(ListItem);
+		}
+	}
+
+	void SBindingMappingsTreeView::RegenerateTreeViewItems()
+	{
+		// Make note of all tree items currently expanded
+		/*
+		TSet<FName> ExpandedAttributeNames;
+		{
+			for (TSharedPtr<FTreeItem> TreeItem : GetAllTreeItems())
+			{
+				if (TreeView->IsItemExpanded(TreeItem))
+				{
+					ExpandedAttributeNames.Add(TreeItem->Name);
+				}
+			}
+		}
+		*/
+
+		// Rebuild items
+		{
+			RootItems.Reset();
+			TMap<FName, TSharedPtr<FBindingAttributeTreeItem>> AttributesItemMap;
+
+			if (const FSkeletonNamedAttributeMapping* const AttributeMapping = SkeletonBinding->FindNamedAttributeMapping(SelectedMappingName))
+			{
+				TArray<TPair<FSkeletonNamedAttribute, TSharedPtr<FBindingAttributeTreeItem>>> NamedAttributesQueue;
+				
+				for (const FSkeletonNamedAttributeMappingEntry& MappingEntry : AttributeMapping->TableData)
+				{
+					const FSkeletonNamedAttribute* NamedAttribute = SkeletonBinding->FindNamedAttribute(MappingEntry.AttributeName);
+					check(NamedAttribute);
+					
+					TSharedPtr<FBindingAttributeTreeItem> MappingTreeItem = MakeShared<FBindingAttributeTreeItem>();
+					MappingTreeItem->AttributeName = NamedAttribute->Name;
+					MappingTreeItem->TreeView = this;
+					MappingTreeItem->bHasValue = true;
+
+					const FSkeletonNamedAttributeMappingType* const AttributeMappingType = AttributeMapping->MappingTypes.FindByPredicate([NamedAttribute](const FSkeletonNamedAttributeMappingType& MappingType)
+					{
+						return MappingType.SourceType == NamedAttribute->Type;
+					});
+
+					if (AttributeMappingType)
+					{
+						MappingTreeItem->ValueType = AttributeMappingType->TargetType;
+						MappingTreeItem->AttributeType = FText::Format(LOCTEXT("AttributeType", "{0} -> {1}"), FText::FromName(NamedAttribute->Type.GetFName()), FText::FromName(AttributeMappingType->TargetType.GetFName()));
+					}
+					else
+					{
+						MappingTreeItem->ValueType = nullptr;
+						MappingTreeItem->AttributeType = FText::Format(LOCTEXT("AttributeTypeNone", "{0} -> None"), FText::FromName(NamedAttribute->Type.GetFName()));
+					}
+
+					AttributesItemMap.Add(NamedAttribute->Name, MappingTreeItem);
+					NamedAttributesQueue.Add(TPair<FSkeletonNamedAttribute, TSharedPtr<FBindingAttributeTreeItem>>(*NamedAttribute, MappingTreeItem));
+				}
+
+				while (!NamedAttributesQueue.IsEmpty())
+				{
+					const TPair<FSkeletonNamedAttribute, TSharedPtr<FBindingAttributeTreeItem>> Attribute = NamedAttributesQueue.Pop();
+					const FSkeletonNamedAttribute NamedAttribute = Attribute.Key;
+					TSharedPtr<FBindingAttributeTreeItem> AttributeTreeItem = Attribute.Value;
+
+					FName AncestorName = NamedAttribute.ParentName;
+					TSharedPtr<FBindingAttributeTreeItem> ClosestAncestorItem = nullptr;
+					while (AncestorName != NAME_None)
+					{
+						TSharedPtr<FBindingAttributeTreeItem>* FoundClosestAncestorItem = AttributesItemMap.Find(AncestorName);
+						if (FoundClosestAncestorItem)
+						{
+							ClosestAncestorItem = *FoundClosestAncestorItem;
+							break;
+						}
+
+						if (const FSkeletonNamedAttribute* const Parent = SkeletonBinding->FindNamedAttribute(AncestorName))
+						{
+							AncestorName = Parent->ParentName;
+						}
+						else
+						{
+							break;
+						}
+					}
+
+					if (ClosestAncestorItem)
+					{
+						AttributeTreeItem->Parent = ClosestAncestorItem;
+						ClosestAncestorItem->ChildrenAttributes.Add(AttributeTreeItem);
+					}
+					else
+					{
+						AttributeTreeItem->Parent = nullptr;
+						RootItems.Add(AttributeTreeItem);
+					}
+				}
+			}
+		}
+
+		// Update tree view and restore tree item expanded states
+		{
+			check(TreeView);
+			TreeView->RequestTreeRefresh();
+
+			for (TSharedPtr<ITreeItem> TreeItem : GetAllTreeItems())
+			{
+				TreeView->SetItemExpansion(TreeItem, true);
+			}
+
+			/*
+			for (TSharedPtr<FTreeItem> TreeItem : GetAllTreeItems())
+			{
+				if (ExpandedAttributeNames.Contains(TreeItem->Name))
+				{
+					TreeView->SetItemExpansion(TreeItem, true);
+				}
+			}
+			*/
+		}
+	}
+}
+
+#undef LOCTEXT_NAMESPACE
\ No newline at end of file
diff --git a/Engine/Plugins/Experimental/Animation/SkeletonTemplateFramework/Source/Editor/Private/SkeletonTemplateFramework/SkeletonBindingNamedAttributeMappingsEditor.h b/Engine/Plugins/Experimental/Animation/SkeletonTemplateFramework/Source/Editor/Private/SkeletonTemplateFramework/SkeletonBindingNamedAttributeMappingsEditor.h
new file mode 100644
--- /dev/null
+++ b/Engine/Plugins/Experimental/Animation/SkeletonTemplateFramework/Source/Editor/Private/SkeletonTemplateFramework/SkeletonBindingNamedAttributeMappingsEditor.h
@@ -0,0 +1,97 @@
+// Copyright Epic Games, Inc. All Rights Reserved.
+
+#pragma once
+
+#include "CoreMinimal.h"
+#include "Widgets/SCompoundWidget.h"
+
+class IDetailsView;
+class ITableRow;
+class STableViewBase;
+template<typename ItemType> class SListView;
+template<typename ItemType> class STreeView;
+class USkeletonBinding;
+
+class ISkeletonBindingEditorToolkit;
+
+namespace UE::Anim::STF
+{
+	class SBindingMappingTableRow;
+	
+	class SBindingMappingsTreeView : public SCompoundWidget
+	{
+	public:
+		friend SBindingMappingTableRow;
+		
+		SLATE_BEGIN_ARGS(SBindingMappingsTreeView) { }
+		SLATE_END_ARGS()
+
+		struct FListItem
+		{
+			FName Name;
+			FName SourceSetName;
+
+			DECLARE_DELEGATE(FOnRenameRequested);
+			FOnRenameRequested OnRenameRequested;
+		};
+		
+		struct ITreeItem
+		{
+			virtual ~ITreeItem() = default;
+			
+			FName AttributeName;
+			FText AttributeType;
+			bool bHasValue;
+			SBindingMappingsTreeView* TreeView;
+			
+			virtual void GetChildren(TArray<TSharedPtr<ITreeItem>>& OutChildren) const = 0;
+		};
+
+		friend ITreeItem;
+		
+		void Construct(const FArguments& InArgs, TObjectPtr<USkeletonBinding> InSkeletonBinding, ISkeletonBindingEditorToolkit* Toolkit);
+
+		void OnNamedAttributeSetsChanged();
+
+	private:
+		TSharedRef<ITableRow> ListView_GenerateItemRow(TSharedPtr<FListItem> Item, const TSharedRef<STableViewBase>& OwnerTable);
+
+		TSharedPtr<SWidget> ListView_HandleContextMenuOpening();
+		
+		void ListView_OnSelectionChanged(TSharedPtr<FListItem> InItem, ESelectInfo::Type InSelectInfo);
+
+		void TreeView_HandleGetChildren(TSharedPtr<ITreeItem> InItem, TArray<TSharedPtr<ITreeItem>>& OutChildren) const;
+
+		void TreeView_OnSelectionChanged(TSharedPtr<ITreeItem> InItem, ESelectInfo::Type InSelectInfo);
+		
+		TSharedRef<ITableRow> TreeView_GenerateItemRow(TSharedPtr<ITreeItem> Item, const TSharedRef<STableViewBase>& OwnerTable);
+		
+		TSharedPtr<SWidget> TreeView_HandleContextMenuOpening();
+
+	private:
+		TArray<TSharedPtr<ITreeItem>> GetAllTreeItems();
+
+		void RegenerateListViewItems();
+		
+		void RegenerateTreeViewItems();
+
+	private:
+		TSharedPtr<SListView<TSharedPtr<FListItem>>> ListView;
+
+		TSharedPtr<FListItem> ListView_DeferredRenameRequest;
+
+		TSharedPtr<STreeView<TSharedPtr<ITreeItem>>> TreeView;
+
+		TArray<TSharedPtr<FListItem>> ListItems;
+		
+		FName SelectedMappingName;
+
+		TArray<TSharedPtr<ITreeItem>> RootItems;
+
+		TSharedPtr<ITreeItem> TreeView_DeferredRenameRequest;
+
+		TObjectPtr<USkeletonBinding> SkeletonBinding;
+
+		ISkeletonBindingEditorToolkit* Toolkit;
+	};
+}
\ No newline at end of file
diff --git a/Engine/Plugins/Experimental/Animation/SkeletonTemplateFramework/Source/Editor/Private/SkeletonTemplateFramework/SkeletonBindingNamedAttributeSetsEditor.cpp b/Engine/Plugins/Experimental/Animation/SkeletonTemplateFramework/Source/Editor/Private/SkeletonTemplateFramework/SkeletonBindingNamedAttributeSetsEditor.cpp
new file mode 100644
--- /dev/null
+++ b/Engine/Plugins/Experimental/Animation/SkeletonTemplateFramework/Source/Editor/Private/SkeletonTemplateFramework/SkeletonBindingNamedAttributeSetsEditor.cpp
@@ -0,0 +1,372 @@
+// Copyright Epic Games, Inc. All Rights Reserved.
+
+#include "SkeletonBindingNamedAttributeSetsEditor.h"
+
+#include "DragAndDrop/DecoratedDragDropOp.h"
+#include "Framework/MultiBox/MultiBoxBuilder.h"
+#include "Modules/ModuleManager.h"
+#include "SkeletonTemplateFramework/SkeletonBinding.h"
+#include "SkeletonTemplateFramework/SkeletonTemplateNamedAttributesEditor.h" // TODO: Change for drag-drop
+#include "StructViewerFilter.h"
+#include "StructViewerModule.h"
+#include "Widgets/Input/SComboButton.h"
+#include "Widgets/Views/STreeView.h"
+
+#define LOCTEXT_NAMESPACE "SkeletonBindingNamedAttributeSetsEditor"
+
+namespace UE::Anim::STF
+{
+	class FNamedElementDragDropOp;
+
+	struct FBindingSetsColumns
+	{
+		static const FName NameId;
+		static const FName TypeId;
+	};
+
+	const FName FBindingSetsColumns::NameId(TEXT("Name"));
+	const FName FBindingSetsColumns::TypeId(TEXT("Type"));
+
+	class SBindingSetTableRow : public SMultiColumnTableRow<TSharedPtr<SBindingSetsTreeView::FTreeItem>>
+	{
+	public:
+		DECLARE_DELEGATE_ThreeParams(FOnAddAttributeToSet, FName /* Attribute Set Name */, FName /* Attribute Name */, bool /* Include Children */);
+
+		SLATE_BEGIN_ARGS(SBindingSetTableRow) { }
+			SLATE_ARGUMENT(TSharedPtr<SBindingSetsTreeView::FTreeItem>, TreeItem)
+			SLATE_EVENT(FOnAddAttributeToSet, OnAddAttributeToSet)
+		SLATE_END_ARGS()
+
+		TOptional<EItemDropZone> OnCanAcceptDrop(const FDragDropEvent& DragDropEvent, EItemDropZone DropZone, TSharedPtr<SBindingSetsTreeView::FTreeItem> TargetItem)
+		{
+			using namespace UE::Anim::STF;
+
+			TOptional<EItemDropZone> ReturnedDropZone;
+	
+			const TSharedPtr<FNamedElementDragDropOp> DragDropOp = DragDropEvent.GetOperationAs<FNamedElementDragDropOp>();
+			if (DragDropOp.IsValid() && TreeItem->ItemType == SBindingSetsTreeView::FTreeItem::EType::AttributeSet)
+			{
+        		ReturnedDropZone = EItemDropZone::BelowItem;
+			}
+	
+			return ReturnedDropZone;
+		}
+
+		FReply OnAcceptDrop(const FDragDropEvent& DragDropEvent, EItemDropZone DropZone, TSharedPtr<SBindingSetsTreeView::FTreeItem> TargetItem)
+		{
+			const TSharedPtr<FNamedElementDragDropOp> DragDropOp = DragDropEvent.GetOperationAs<FNamedElementDragDropOp>();
+			if (!DragDropOp.IsValid())
+			{
+				return FReply::Unhandled();
+			}
+
+			const bool bIncludeChildren = DragDropEvent.IsControlDown() || DragDropEvent.IsCommandDown();
+			
+			OnAddAttributeToSet.ExecuteIfBound(TreeItem->Name, DragDropOp->NamedAttribute, bIncludeChildren);
+
+			return FReply::Handled();
+		}
+
+		void Construct(const FArguments& InArgs, const TSharedRef<STableViewBase>& InOwnerTableView)
+		{
+			TreeItem = InArgs._TreeItem;
+			OnAddAttributeToSet = InArgs._OnAddAttributeToSet;
+
+			FSuperRowType::Construct(
+				FSuperRowType::FArguments()
+					.OnCanAcceptDrop(this, &SBindingSetTableRow::OnCanAcceptDrop)
+					.OnAcceptDrop(this, &SBindingSetTableRow::OnAcceptDrop),
+				InOwnerTableView);
+		}
+
+		TSharedRef<SWidget> GenerateWidgetForColumn(const FName& ColumnName) override
+		{
+			if (ColumnName == UE::Anim::STF::FBindingSetsColumns::NameId)
+			{
+
+				TSharedRef<SHorizontalBox> HorizontalBox = SNew(SHorizontalBox)
+					+ SHorizontalBox::Slot()
+						.AutoWidth()
+						[
+							SNew(SExpanderArrow, SharedThis(this))
+								.ShouldDrawWires(true)
+						]
+					+ SHorizontalBox::Slot()
+						.AutoWidth()
+						.Padding(4.0f, 0.0f, 0.0f, 0.0f)
+						[
+							SNew(SImage)
+								.Image(FAppStyle::GetBrush("LevelEditor.Tabs.Outliner"))
+								.Visibility(TreeItem->ItemType == SBindingSetsTreeView::FTreeItem::EType::AttributeSet ? EVisibility::Visible : EVisibility::Collapsed)
+						]
+					+ SHorizontalBox::Slot()
+						.AutoWidth()
+						.Padding(4.0f, 0.0f, 0.0f, 0.0f)
+						[
+							SNew(STextBlock)
+								.Text_Lambda([this]()
+								{
+									return FText::FromName(TreeItem->Name);
+								})
+						];
+
+				return HorizontalBox;
+			}
+			else if (ColumnName == UE::Anim::STF::FBindingSetsColumns::TypeId)
+			{
+				if (TreeItem->ItemType == SBindingSetsTreeView::FTreeItem::EType::AttributeSet)
+				{
+					return SNullWidget::NullWidget;
+				}
+				else
+				{
+					return SNew(STextBlock)
+						.Text_Lambda([this]()
+							{
+								return FText::FromName(TreeItem->Type);
+							});
+				}
+			}
+			else
+			{
+				return SNullWidget::NullWidget;
+			}
+		}
+
+	private:
+		TSharedPtr<SBindingSetsTreeView::FTreeItem> TreeItem;
+		FOnAddAttributeToSet OnAddAttributeToSet;
+	};
+
+	void SBindingSetsTreeView::Construct(const FArguments& InArgs, TObjectPtr<USkeletonBinding> InSkeletonBinding)
+	{
+		SkeletonBinding = InSkeletonBinding;
+
+		ChildSlot
+		[
+			SAssignNew(TreeView, STreeView<TSharedPtr<FTreeItem>>)
+				.TreeItemsSource(&RootItems)
+				.OnGenerateRow(this, &SBindingSetsTreeView::TreeView_GenerateItemRow)
+				.OnGetChildren(this, &SBindingSetsTreeView::TreeView_HandleGetChildren)
+				.OnContextMenuOpening(this, &SBindingSetsTreeView::TreeView_HandleContextMenuOpening)
+				.HighlightParentNodesForSelection(true)
+				.HeaderRow
+				(
+					SNew(SHeaderRow)
+						+SHeaderRow::Column(UE::Anim::STF::FBindingSetsColumns::NameId)
+							.FillWidth(0.5f)
+							.DefaultLabel(LOCTEXT("NameLabel", "Name"))
+						+SHeaderRow::Column(UE::Anim::STF::FBindingSetsColumns::TypeId)
+							.FillWidth(0.5f)
+							.DefaultLabel(LOCTEXT("TypeLabel", "Type"))
+				)
+		];
+
+		RegenerateTreeViewItems();
+
+		// Expand all tree items on construction
+		for (TSharedPtr<FTreeItem> TreeItem : GetAllTreeItems())
+		{
+			TreeView->SetItemExpansion(TreeItem, true);
+		}
+	}
+
+	void SBindingSetsTreeView::TreeView_HandleGetChildren(TSharedPtr<FTreeItem> InItem, TArray<TSharedPtr<FTreeItem>>& OutChildren) const
+	{
+		OutChildren.Append(InItem->Children);
+	}
+
+	TSharedRef<ITableRow> SBindingSetsTreeView::TreeView_GenerateItemRow(TSharedPtr<SBindingSetsTreeView::FTreeItem> Item, const TSharedRef<STableViewBase>& OwnerTable)
+	{
+		if (Item->ItemType == FTreeItem::EType::AttributeSet)
+		{
+			return SNew(SBindingSetTableRow, OwnerTable)
+				.TreeItem(Item)
+				.OnAddAttributeToSet_Lambda([this, Item](const FName AttributeSet, const FName AttributeName, const bool bIncludeChildren)
+					{
+						SkeletonBinding->AddNamedAttributeToSet(AttributeSet, AttributeName);
+						if (bIncludeChildren)
+						{
+							TArray<FName> ChildNamedAttributes;
+							SkeletonBinding->GetChildNamedAttributes(AttributeName, ChildNamedAttributes);
+
+							for (const FName ChildNamedAttribute : ChildNamedAttributes)
+							{
+								SkeletonBinding->AddNamedAttributeToSet(AttributeSet, ChildNamedAttribute);
+							}
+						}
+					
+						TSharedPtr<FTreeItem> Parent = Item->Parent;
+						while (Parent)
+						{
+							TreeView->SetItemExpansion(Parent, true);
+							Parent = Parent->Parent;
+						}
+						
+						RegenerateTreeViewItems();
+					});
+		}
+		else
+		{
+			return SNew(SBindingSetTableRow, OwnerTable)
+				.TreeItem(Item);
+		}
+	}
+
+	TSharedPtr<SWidget> SBindingSetsTreeView::TreeView_HandleContextMenuOpening()
+	{
+		FMenuBuilder MenuBuilder(true, nullptr);
+
+		TArray<TSharedPtr<FTreeItem>> SelectedItems;
+		TreeView->GetSelectedItems(SelectedItems);
+
+		if (!SelectedItems.IsEmpty())
+		{
+			TSharedPtr<FTreeItem> SelectedItem = SelectedItems[0];
+
+			switch (SelectedItem->ItemType)
+			{
+				case FTreeItem::EType::Attribute:
+				{
+					MenuBuilder.AddMenuEntry(
+						LOCTEXT("AddNamedAttribute", "Remove Named Attribute"),
+						LOCTEXT("AddNamedAttribute_Tooltip", "Remove the selected named attribute from the set"),
+						FSlateIcon(FAppStyle::GetAppStyleSetName(), "Icons.Minus"),
+						FUIAction(
+							FExecuteAction::CreateLambda([this, SelectedItem]()
+							{
+								SkeletonBinding->RemoveNamedAttributeFromSet(SelectedItem->SetItem->Name, SelectedItem->Name);
+								RegenerateTreeViewItems();
+							})
+						));
+
+					break;
+				}
+			}
+		}
+
+		return MenuBuilder.MakeWidget();
+	}
+
+	TArray<TSharedPtr<UE::Anim::STF::SBindingSetsTreeView::FTreeItem>> SBindingSetsTreeView::GetAllTreeItems()
+	{
+		using namespace UE::Anim::STF;
+
+		TArray<TSharedPtr<SBindingSetsTreeView::FTreeItem>> AllItems;
+		AllItems.Append(RootItems);
+	
+		for (int32 Index = 0; Index < AllItems.Num(); ++Index)
+		{
+			AllItems.Append(AllItems[Index]->Children);
+		}
+
+		return AllItems;
+	}
+
+	void SBindingSetsTreeView::RegenerateTreeViewItems()
+	{
+		// Make note of all tree items currently expanded
+		TSet<FName> ExpandedAttributeNames;
+		{
+			for (TSharedPtr<FTreeItem> TreeItem : GetAllTreeItems())
+			{
+				if (TreeView->IsItemExpanded(TreeItem))
+				{
+					ExpandedAttributeNames.Add(TreeItem->Name);
+				}
+			}
+		}
+
+		// Rebuild items
+		{
+			RootItems.Reset();
+	
+			for (const FSkeletonNamedAttributeSet& NamedAttributeSet : SkeletonBinding->GetNamedAttributeSets())
+			{
+				TMap<FName, TSharedPtr<FTreeItem>> AttributesItemMap;
+
+				TArray<TPair<FSkeletonNamedAttribute, TSharedPtr<FTreeItem>>> NamedAttributesQueue;
+
+				TSharedPtr<FTreeItem> SetItem = MakeShared<FTreeItem>();
+				SetItem->Name = NamedAttributeSet.Name;
+				SetItem->ItemType = FTreeItem::EType::AttributeSet;
+				SetItem->Type = NAME_None;
+
+				for (const FName& AttributeName : NamedAttributeSet.NamedAttributes)
+				{
+					const FSkeletonNamedAttribute* const NamedAttribute = SkeletonBinding->FindNamedAttribute(AttributeName);
+					if (!ensure(NamedAttribute))
+					{
+						continue;
+					}
+
+					TSharedPtr<FTreeItem> AttributeItem = MakeShared<FTreeItem>();
+					AttributeItem->Name = AttributeName;
+					AttributeItem->ItemType = FTreeItem::EType::Attribute;
+					AttributeItem->Type = NamedAttribute->Type.GetFName();
+					AttributeItem->SetItem = SetItem;
+
+					AttributesItemMap.Add(AttributeName, AttributeItem);
+					NamedAttributesQueue.Add(TPair<FSkeletonNamedAttribute, TSharedPtr<FTreeItem>>(*NamedAttribute, AttributeItem));
+				}
+
+				while (!NamedAttributesQueue.IsEmpty())
+				{
+					const TPair<FSkeletonNamedAttribute, TSharedPtr<FTreeItem>> Attribute = NamedAttributesQueue.Pop();
+					const FSkeletonNamedAttribute NamedAttribute = Attribute.Key;
+					TSharedPtr<FTreeItem> AttributeTreeItem = Attribute.Value;
+
+					FName AncestorName = NamedAttribute.ParentName;
+					TSharedPtr<FTreeItem> ClosestAncestorItem = nullptr;
+					while (AncestorName != NAME_None)
+					{
+						TSharedPtr<FTreeItem>* FoundClosestAncestorItem = AttributesItemMap.Find(AncestorName);
+						if (FoundClosestAncestorItem)
+						{
+							ClosestAncestorItem = *FoundClosestAncestorItem;
+							break;
+						}
+
+						if (const FSkeletonNamedAttribute* const Parent = SkeletonBinding->FindNamedAttribute(AncestorName))
+						{
+							AncestorName = Parent->ParentName;
+						}
+						else
+						{
+							break;
+						}
+					}
+
+					if (ClosestAncestorItem)
+					{
+						AttributeTreeItem->Parent = ClosestAncestorItem;
+						ClosestAncestorItem->Children.Add(AttributeTreeItem);
+					}
+					else
+					{
+						AttributeTreeItem->Parent = SetItem;
+						SetItem->Children.Add(AttributeTreeItem);
+					}
+				}
+
+				RootItems.Add(SetItem);
+			}
+		}
+
+		// Update tree view and restore tree item expanded states
+		{
+			check(TreeView);
+			TreeView->RequestTreeRefresh();
+
+			for (TSharedPtr<FTreeItem> TreeItem : GetAllTreeItems())
+			{
+				if (ExpandedAttributeNames.Contains(TreeItem->Name))
+				{
+					TreeView->SetItemExpansion(TreeItem, true);
+				}
+			}
+		}
+	}
+}
+
+#undef LOCTEXT_NAMESPACE
\ No newline at end of file
diff --git a/Engine/Plugins/Experimental/Animation/SkeletonTemplateFramework/Source/Editor/Private/SkeletonTemplateFramework/SkeletonBindingNamedAttributeSetsEditor.h b/Engine/Plugins/Experimental/Animation/SkeletonTemplateFramework/Source/Editor/Private/SkeletonTemplateFramework/SkeletonBindingNamedAttributeSetsEditor.h
new file mode 100644
--- /dev/null
+++ b/Engine/Plugins/Experimental/Animation/SkeletonTemplateFramework/Source/Editor/Private/SkeletonTemplateFramework/SkeletonBindingNamedAttributeSetsEditor.h
@@ -0,0 +1,61 @@
+// Copyright Epic Games, Inc. All Rights Reserved.
+
+#pragma once
+
+#include "CoreMinimal.h"
+#include "Widgets/SCompoundWidget.h"
+
+class ITableRow;
+class STableViewBase;
+template<typename ItemType> class STreeView;
+class USkeletonBinding;
+
+namespace UE::Anim::STF
+{
+	class SBindingSetsTreeView : public SCompoundWidget
+	{
+	public:
+		SLATE_BEGIN_ARGS(SBindingSetsTreeView) { }
+		SLATE_END_ARGS()
+
+		struct FTreeItem
+		{
+			enum class EType
+			{
+				AttributeSet,
+				Attribute,
+			};
+
+			EType ItemType;
+			FName Name;
+			FName Type;
+			TSharedPtr<FTreeItem> SetItem;
+			TSharedPtr<FTreeItem> Parent;
+			TArray<TSharedPtr<FTreeItem>> Children;
+
+			DECLARE_DELEGATE(FOnRenameRequested);
+			FOnRenameRequested OnRenameRequested;
+		};
+
+		void Construct(const FArguments& InArgs, TObjectPtr<USkeletonBinding> InSkeletonBinding);
+
+	private:
+		void TreeView_HandleGetChildren(TSharedPtr<FTreeItem> InItem, TArray<TSharedPtr<FTreeItem>>& OutChildren) const;
+		
+		TSharedRef<ITableRow> TreeView_GenerateItemRow(TSharedPtr<FTreeItem> Item, const TSharedRef<STableViewBase>& OwnerTable);
+		
+		TSharedPtr<SWidget> TreeView_HandleContextMenuOpening();
+
+	private:
+		TArray<TSharedPtr<FTreeItem>> GetAllTreeItems();
+
+		void RegenerateTreeViewItems();
+
+	private:
+		TSharedPtr<STreeView<TSharedPtr<FTreeItem>>> TreeView;
+
+		TArray<TSharedPtr<FTreeItem>> RootItems;
+
+		TObjectPtr<USkeletonBinding> SkeletonBinding;
+	};
+}
\ No newline at end of file
diff --git a/Engine/Plugins/Experimental/Animation/SkeletonTemplateFramework/Source/Editor/Private/SkeletonTemplateFramework/SkeletonBindingNamedAttributesEditor.cpp b/Engine/Plugins/Experimental/Animation/SkeletonTemplateFramework/Source/Editor/Private/SkeletonTemplateFramework/SkeletonBindingNamedAttributesEditor.cpp
new file mode 100644
--- /dev/null
+++ b/Engine/Plugins/Experimental/Animation/SkeletonTemplateFramework/Source/Editor/Private/SkeletonTemplateFramework/SkeletonBindingNamedAttributesEditor.cpp
@@ -0,0 +1,317 @@
+// Copyright Epic Games, Inc. All Rights Reserved.
+
+#include "SkeletonBindingNamedAttributesEditor.h"
+
+#include "Animation/Skeleton.h"
+#include "Framework/MultiBox/MultiBoxBuilder.h"
+#include "Modules/ModuleManager.h"
+#include "SkeletonTemplateFramework/SkeletonBinding.h"
+#include "SkeletonTemplateFramework/SkeletonTemplateNamedAttributesEditor.h"
+#include "StructViewerFilter.h"
+#include "StructViewerModule.h"
+#include "Widgets/Input/SComboBox.h"
+#include "Widgets/Text/SInlineEditableTextBlock.h"
+#include "Widgets/Views/STreeView.h"
+
+#define LOCTEXT_NAMESPACE "SkeletonBindingNamedAttributesEditor"
+
+namespace UE::Anim::STF
+{
+	struct FAttributeBindingsColumns
+	{
+		static const FName NameId;
+		static const FName TypeId;
+		static const FName BindingId;
+	};
+
+	const FName FAttributeBindingsColumns::NameId(TEXT("Name"));
+	const FName FAttributeBindingsColumns::TypeId(TEXT("Type"));
+	const FName FAttributeBindingsColumns::BindingId(TEXT("Binding"));
+
+	class SNamedElementBindingTableRow : public SMultiColumnTableRow<TSharedPtr<SAttributeBindingsTreeView::FTreeItem>>
+	{
+	public:
+		DECLARE_DELEGATE_TwoParams(FOnBindingSelected, FName /* Old Binding */, FName /* New Binding */);
+		
+		SLATE_BEGIN_ARGS(SNamedElementBindingTableRow) { }
+			SLATE_ARGUMENT(TSharedPtr<SAttributeBindingsTreeView::FTreeItem>, NamedElementTreeItem)
+			SLATE_ARGUMENT(TArray<TSharedPtr<FName>>*, BindingOptions)
+			SLATE_EVENT(FOnBindingSelected, OnBindingSelected)
+		SLATE_END_ARGS()
+
+		FReply OnDragDetected(const FGeometry& MyGeometry, const FPointerEvent& MouseEvent)
+		{
+			if (MouseEvent.IsMouseButtonDown(EKeys::LeftMouseButton))
+			{
+				const TSharedRef<FNamedElementDragDropOp> DragDropOp = FNamedElementDragDropOp::New(NamedElementTreeItem->Name);
+				return FReply::Handled().BeginDragDrop(DragDropOp);
+			}
+
+			return FReply::Unhandled();
+		}
+		
+		void Construct(const FArguments& InArgs, const TSharedRef<STableViewBase>& InOwnerTableView)
+		{
+			NamedElementTreeItem = InArgs._NamedElementTreeItem;
+			BindingOptions = InArgs._BindingOptions;
+			OnBindingSelected = InArgs._OnBindingSelected;
+
+			FSuperRowType::Construct(
+				FSuperRowType::FArguments()
+					.OnDragDetected(this, &SNamedElementBindingTableRow::OnDragDetected),
+				InOwnerTableView);
+		}
+
+		TSharedRef<SWidget> GenerateWidgetForColumn(const FName& ColumnName) override
+		{
+			if (ColumnName == UE::Anim::STF::FAttributeBindingsColumns::NameId)
+			{
+				TSharedPtr<SInlineEditableTextBlock> InlineWidget;
+
+				TSharedRef<SHorizontalBox> HorizontalBox = SNew(SHorizontalBox)
+					+ SHorizontalBox::Slot()
+						.AutoWidth()
+						[
+							SNew(SExpanderArrow, SharedThis(this))
+								.ShouldDrawWires(true)
+						]
+					+ SHorizontalBox::Slot()
+						.AutoWidth()
+						[
+							SNew(STextBlock)
+								.Text(FText::FromName(NamedElementTreeItem->Name))
+						];
+
+				return HorizontalBox;
+			}
+			else if (ColumnName == UE::Anim::STF::FAttributeBindingsColumns::TypeId)
+			{
+				return SNew(STextBlock)
+					.Text_Lambda([this]()
+						{
+							return FText::FromName(NamedElementTreeItem->Type);
+						});
+			}
+			else if (ColumnName == UE::Anim::STF::FAttributeBindingsColumns::BindingId)
+			{
+				TSharedPtr<FName> InitialValue;
+				for (TSharedPtr<FName> Option : *BindingOptions)
+				{
+					if (*Option == NamedElementTreeItem->Binding)
+					{
+						InitialValue = Option;
+						break;
+					}
+				}
+				
+				return SNew(SComboBox<TSharedPtr<FName>>)
+					.InitiallySelectedItem(InitialValue)
+					.OptionsSource(BindingOptions)
+					.OnGenerateWidget_Lambda([](TSharedPtr<FName> Name)
+						{
+							return SNew(STextBlock)
+								.Text(FText::FromName(*Name));
+						})
+					.OnSelectionChanged_Lambda([this](TSharedPtr<FName> NewSelection, ESelectInfo::Type)
+						{
+							OnBindingSelected.ExecuteIfBound(NamedElementTreeItem->Binding, *NewSelection);
+						})
+					[
+						SNew(STextBlock)
+							.Text_Lambda([this]()
+							{
+								return FText::FromName(NamedElementTreeItem->Binding);
+							})
+					];
+			}
+			else
+			{
+				return SNullWidget::NullWidget;
+			}
+		}
+
+	private:
+		FOnBindingSelected OnBindingSelected;
+		
+		TArray<TSharedPtr<FName>>* BindingOptions;
+		
+		TSharedPtr<SAttributeBindingsTreeView::FTreeItem> NamedElementTreeItem;
+	};
+
+	void SAttributeBindingsTreeView::Construct(const FArguments& InArgs, TObjectPtr<USkeletonBinding> InSkeletonBinding)
+	{
+		SkeletonBinding = InSkeletonBinding;
+
+		ChildSlot
+		[
+			SAssignNew(TreeView, STreeView<TSharedPtr<FTreeItem>>)
+				.TreeItemsSource(&RootItems)
+				.OnGenerateRow(this, &SAttributeBindingsTreeView::TreeView_GenerateItemRow)
+				.OnGetChildren(this, &SAttributeBindingsTreeView::TreeView_HandleGetChildren)
+				.OnContextMenuOpening(this, &SAttributeBindingsTreeView::TreeView_HandleContextMenuOpening)
+				.HighlightParentNodesForSelection(true)
+				.HeaderRow
+				(
+					SNew(SHeaderRow)
+						+ SHeaderRow::Column(UE::Anim::STF::FAttributeBindingsColumns::NameId)
+							.FillWidth(0.5f)
+							.DefaultLabel(LOCTEXT("NameLabel", "Name"))
+						+ SHeaderRow::Column(UE::Anim::STF::FAttributeBindingsColumns::TypeId)
+							.FillWidth(0.1f)
+							.DefaultLabel(LOCTEXT("TypeLabel", "Type"))
+						+ SHeaderRow::Column(UE::Anim::STF::FAttributeBindingsColumns::BindingId)
+							.FillWidth(0.4f)
+							.DefaultLabel(LOCTEXT("BindingLabel", "Binding"))
+				)
+		];
+
+		RegenerateTreeViewItems();
+
+		BindingOptions.Add(MakeShared<FName>(NAME_None));
+		const TArray<FSkeletonNamedAttribute>& NamedAttributes = SkeletonBinding->GetUnboundSchemaNamedAttributes();
+		for (const FSkeletonNamedAttribute& NamedAttribute : NamedAttributes)
+		{
+			BindingOptions.Add(MakeShared<FName>(NamedAttribute.Name));
+		}
+
+		// Expand all tree items on construction
+		for (TSharedPtr<FTreeItem> TreeItem : GetAllTreeItems())
+		{
+			TreeView->SetItemExpansion(TreeItem, true);
+		}
+	}
+
+	void SAttributeBindingsTreeView::Tick(const FGeometry& AllottedGeometry, const double InCurrentTime, const float InDeltaTime)
+	{
+		if (bRequestRegenerateTreeItems)
+		{
+			RegenerateTreeViewItems();
+			bRequestRegenerateTreeItems = false;
+		}
+	}
+
+	void SAttributeBindingsTreeView::TreeView_HandleGetChildren(TSharedPtr<FTreeItem> InItem, TArray<TSharedPtr<FTreeItem>>& OutChildren) const
+	{
+		OutChildren.Append(InItem->Children);
+	}
+
+	TSharedRef<ITableRow> SAttributeBindingsTreeView::TreeView_GenerateItemRow(TSharedPtr<SAttributeBindingsTreeView::FTreeItem> Item, const TSharedRef<STableViewBase>& OwnerTable)
+	{
+		return SNew(SNamedElementBindingTableRow, OwnerTable)
+			.NamedElementTreeItem(Item)
+			.BindingOptions(&BindingOptions)
+			.OnBindingSelected_Lambda([this, Item](const FName OldTemplateNamedAttribute, const FName NewTemplateNamedAttribute)
+			{
+				SkeletonBinding->UnbindTemplateNamedAttribute(OldTemplateNamedAttribute);
+				SkeletonBinding->UnbindTemplateNamedAttribute(NewTemplateNamedAttribute);
+
+				SkeletonBinding->BindAttribute(Item->Name, NewTemplateNamedAttribute);
+				
+				RequestRegenerateTreeItems();
+			});
+	}
+
+	TSharedPtr<SWidget> SAttributeBindingsTreeView::TreeView_HandleContextMenuOpening()
+	{
+		FMenuBuilder MenuBuilder(true, nullptr);
+		return MenuBuilder.MakeWidget();
+	}
+
+	TArray<TSharedPtr<UE::Anim::STF::SAttributeBindingsTreeView::FTreeItem>> SAttributeBindingsTreeView::GetAllTreeItems()
+	{
+		using namespace UE::Anim::STF;
+
+		TArray<TSharedPtr<SAttributeBindingsTreeView::FTreeItem>> AllItems;
+		AllItems.Append(RootItems);
+	
+		for (int32 Index = 0; Index < AllItems.Num(); ++Index)
+		{
+			AllItems.Append(AllItems[Index]->Children);
+		}
+
+		return AllItems;
+	}
+
+	void SAttributeBindingsTreeView::RegenerateTreeViewItems()
+	{
+		// Make note of all tree items currently expanded
+		TSet<FName> ExpandedAttributeNames;
+		{
+			for (TSharedPtr<FTreeItem> TreeItem : GetAllTreeItems())
+			{
+				if (TreeView->IsItemExpanded(TreeItem))
+				{
+					ExpandedAttributeNames.Add(TreeItem->Name);
+				}
+			}
+		}
+
+		// Rebuild items
+		{
+			RootItems.Reset();
+	
+			TMap<FName, TSharedPtr<FTreeItem>> ItemMap;
+
+			for (const FSkeletonNamedAttribute& NamedAttribute : SkeletonBinding->GetNamedAttributes())
+			{
+				TSharedPtr<FTreeItem> Item = MakeShared<FTreeItem>();
+				Item->Name = NamedAttribute.Name;
+				Item->Type = NamedAttribute.Type ? (NamedAttribute.Type->GetFName()) : NAME_None;
+
+				if (const FSkeletonAttributeBinding* const Binding = SkeletonBinding->FindAttributeBinding(NamedAttribute.Name))
+				{
+					Item->Binding = Binding->AttributeName;
+				}
+
+				const TSharedPtr<FTreeItem>* ParentItem = ItemMap.Find(NamedAttribute.ParentName);
+				if (ParentItem)
+				{
+					(*ParentItem)->Children.Add(Item);
+				}
+				else
+				{
+					RootItems.Add(Item);
+				}
+
+				ItemMap.Add(Item->Name, Item);
+			}
+		}
+
+		// Update tree view and restore tree item expanded states
+		{
+			check(TreeView);
+			TreeView->RequestTreeRefresh();
+
+			for (TSharedPtr<FTreeItem> TreeItem : GetAllTreeItems())
+			{
+				if (ExpandedAttributeNames.Contains(TreeItem->Name))
+				{
+					TreeView->SetItemExpansion(TreeItem, true);
+				}
+			}
+		}
+	}
+
+	TSharedPtr<UE::Anim::STF::SAttributeBindingsTreeView::FTreeItem> SAttributeBindingsTreeView::FindTreeItemByTemplateNamedAttribute(const FName TemplateNamedAttribute)
+	{
+		TArray<TSharedPtr<FTreeItem>> AllTreeItems = GetAllTreeItems();
+		
+		for (TSharedPtr<FTreeItem> TreeItem : AllTreeItems)
+		{
+			if (TreeItem->Binding == TemplateNamedAttribute)
+			{
+				return TreeItem;
+			}
+		}
+
+		return nullptr;
+	}
+
+	void SAttributeBindingsTreeView::RequestRegenerateTreeItems()
+	{
+		bRequestRegenerateTreeItems = true;
+	}
+
+}
+
+#undef LOCTEXT_NAMESPACE
\ No newline at end of file
diff --git a/Engine/Plugins/Experimental/Animation/SkeletonTemplateFramework/Source/Editor/Private/SkeletonTemplateFramework/SkeletonBindingNamedAttributesEditor.h b/Engine/Plugins/Experimental/Animation/SkeletonTemplateFramework/Source/Editor/Private/SkeletonTemplateFramework/SkeletonBindingNamedAttributesEditor.h
new file mode 100644
--- /dev/null
+++ b/Engine/Plugins/Experimental/Animation/SkeletonTemplateFramework/Source/Editor/Private/SkeletonTemplateFramework/SkeletonBindingNamedAttributesEditor.h
@@ -0,0 +1,76 @@
+// Copyright Epic Games, Inc. All Rights Reserved.
+
+#pragma once
+
+#include "CoreMinimal.h"
+#include "DragAndDrop/DecoratedDragDropOp.h"
+#include "Widgets/SCompoundWidget.h"
+#include "SkeletonTemplateFramework/SkeletonBinding.h"
+
+class ITableRow;
+class STableViewBase;
+template<typename ItemType> class STreeView;
+class USkeletonBinding;
+
+namespace UE::Anim::STF
+{
+	class SAttributeBindingsTreeView : public SCompoundWidget
+	{
+	public:
+		SLATE_BEGIN_ARGS(SAttributeBindingsTreeView) { }
+		SLATE_END_ARGS()
+
+		struct FTreeItem
+		{
+			FName Name;
+			FName Type;
+			FName Binding;
+			TArray<TSharedPtr<FTreeItem>> Children;
+
+			DECLARE_DELEGATE(FOnRenameRequested);
+			FOnRenameRequested OnRenameRequested;
+		};
+
+		void Construct(const FArguments& InArgs, TObjectPtr<USkeletonBinding> InSkeletonBinding);
+
+		virtual void Tick(const FGeometry& AllottedGeometry, const double InCurrentTime, const float InDeltaTime) override;
+		
+	private:
+		
+		void TreeView_HandleGetChildren(TSharedPtr<FTreeItem> InItem, TArray<TSharedPtr<FTreeItem>>& OutChildren) const;
+		
+		void TreeView_OnItemScrolledIntoView(TSharedPtr<FTreeItem> InItem, const TSharedPtr<ITableRow>& InWidget);
+
+		TSharedRef<ITableRow> TreeView_GenerateItemRow(TSharedPtr<FTreeItem> Item, const TSharedRef<STableViewBase>& OwnerTable);
+		
+		TSharedPtr<SWidget> TreeView_HandleContextMenuOpening();
+
+	private:
+		TArray<TSharedPtr<FTreeItem>> GetAllTreeItems();
+
+		void RegenerateTreeViewItems();
+
+		void HandleAddAttribute();
+
+		void HandleRenameAttribute();
+
+		void HandleDeleteAttribute();
+
+		void RequestRegenerateTreeItems();
+
+		TSharedPtr<FTreeItem> FindTreeItemByTemplateNamedAttribute(const FName TemplateNamedAttribute);
+
+	private:
+		TSharedPtr<STreeView<TSharedPtr<FTreeItem>>> TreeView;
+
+		TArray<TSharedPtr<FTreeItem>> RootItems;
+
+		TSharedPtr<FTreeItem> DeferredRenameRequest;
+
+		TObjectPtr<USkeletonBinding> SkeletonBinding;
+
+		TArray<TSharedPtr<FName>> BindingOptions;
+
+		bool bRequestRegenerateTreeItems;
+	};
+}
\ No newline at end of file
diff --git a/Engine/Plugins/Experimental/Animation/SkeletonTemplateFramework/Source/Runtime/Private/SkeletonTemplateFramework/SkeletonBinding.cpp b/Engine/Plugins/Experimental/Animation/SkeletonTemplateFramework/Source/Runtime/Private/SkeletonTemplateFramework/SkeletonBinding.cpp
--- a/Engine/Plugins/Experimental/Animation/SkeletonTemplateFramework/Source/Runtime/Private/SkeletonTemplateFramework/SkeletonBinding.cpp
+++ b/Engine/Plugins/Experimental/Animation/SkeletonTemplateFramework/Source/Runtime/Private/SkeletonTemplateFramework/SkeletonBinding.cpp
@@ -2,6 +2,7 @@
 
 #include "SkeletonTemplateFramework/SkeletonBinding.h"
 
+#include "SkeletonTemplateFramework/NamedElementTypes.h"
 #include "Animation/Skeleton.h"
 
 #include UE_INLINE_GENERATED_CPP_BY_NAME(SkeletonBinding)
@@ -21,11 +22,252 @@ bool USkeletonBinding::InitializeFrom(const TObjectPtr<USkeletonTemplate> InTemp
 	Template = InTemplate;
 	Skeleton = InSkeleton;
 	AttributeBindings.Reset();
-	NamedAttributes = InTemplate->GetNamedAttributes();
-	NamedAttributeSets = InTemplate->GetNamedAttributeSets();
-	NamedAttributeMappings = InTemplate->GetNamedAttributeMappings();
 	CachedSkeletonGuid = InSkeleton->GetGuid();
+	
+	NamedAttributes.Reset();
+	{
+		const FReferenceSkeleton& ReferenceSkeleton = InSkeleton->GetReferenceSkeleton();
+		const int32 BoneNum = ReferenceSkeleton.GetNum();
+
+		for (int32 BoneIndex = 0; BoneIndex < BoneNum; ++BoneIndex)
+		{
+			const int32 ParentIndex = ReferenceSkeleton.GetParentIndex(BoneIndex);
+			const FName ParentName = ParentIndex != INDEX_NONE ? ReferenceSkeleton.GetBoneName(ParentIndex) : NAME_None;
+			
+			FSkeletonNamedAttribute NewNamedAttribute;
+			NewNamedAttribute.Name = ReferenceSkeleton.GetBoneName(BoneIndex);
+			NewNamedAttribute.Type = USkeletonTemplateTransform::StaticClass();
+			NewNamedAttribute.ParentName = ParentName;
+			
+			NamedAttributes.Add(NewNamedAttribute);
+		}
+	}
+	
+	for (const FSkeletonNamedAttributeSet& TemplateAttributeSet : InTemplate->GetNamedAttributeSets())
+	{
+		FSkeletonNamedAttributeSet BindingAttributeSet;
+		BindingAttributeSet.Name = TemplateAttributeSet.Name;
+		
+		NamedAttributeSets.Add(BindingAttributeSet);
+	}
+	
+	for (const FSkeletonNamedAttributeMapping& TemplateAttributeMapping : InTemplate->GetNamedAttributeMappings())
+	{
+		FSkeletonNamedAttributeMapping BindingAttributeMapping;
+		BindingAttributeMapping.Name = TemplateAttributeMapping.Name;
+		BindingAttributeMapping.SourceAttributeSet = TemplateAttributeMapping.SourceAttributeSet;
+		BindingAttributeMapping.MappingTypes = TemplateAttributeMapping.MappingTypes;
+		
+		NamedAttributeMappings.Add(BindingAttributeMapping);
+	}
+	
 	// TODO: Set CachedParentTemplateHash
 
 	return true;
+}
+
+const FSkeletonNamedAttribute* const USkeletonBinding::FindNamedAttribute(const FName AttributeName) const
+{
+	return NamedAttributes.FindByPredicate([AttributeName](const FSkeletonNamedAttribute& Candidate)
+		{
+			return Candidate.Name == AttributeName;
+		});
+}
+	
+
+TObjectPtr<USkeleton> USkeletonBinding::GetSkeleton() const
+{
+	return Skeleton;
+}
+
+TObjectPtr<USkeletonTemplate> USkeletonBinding::GetSkeletonTemplate() const
+{
+	return Template;
+}
+
+const TArray<FSkeletonNamedAttribute>& USkeletonBinding::GetNamedAttributes() const
+{
+	return NamedAttributes;
+}
+
+const TArray<FSkeletonNamedAttributeSet>& USkeletonBinding::GetNamedAttributeSets() const
+{
+	return NamedAttributeSets;
+}
+
+const TArray<FSkeletonNamedAttributeMapping>& USkeletonBinding::GetNamedAttributeMappings() const
+{
+	return NamedAttributeMappings;
+}
+
+TArray<FSkeletonNamedAttribute> USkeletonBinding::GetUnboundSchemaNamedAttributes() const
+{
+	return Template->GetNamedAttributes().FilterByPredicate([this](const FSkeletonNamedAttribute& SchemaAttribute)
+	{
+		const int32 BoundAttributeIndex = AttributeBindings.IndexOfByPredicate([SchemaAttribute](const FSkeletonAttributeBinding& AttributeBinding)
+		{
+			return AttributeBinding.AttributeName == SchemaAttribute.Name;
+		});
+
+		return BoundAttributeIndex == INDEX_NONE;
+	});
+}
+
+const FSkeletonAttributeBinding* const USkeletonBinding::FindAttributeBinding(const FName BindingAttributeName)  const
+{
+	const int32 AttributeIndex = FindAttributeBindingIndex(BindingAttributeName);
+	if (AttributeIndex == INDEX_NONE)
+	{
+		return nullptr;
+	}
+
+	return &AttributeBindings[AttributeIndex];
+}
+
+int32 USkeletonBinding::FindAttributeBindingIndex(const FName BindingAttributeName) const
+{
+	return AttributeBindings.IndexOfByPredicate([BindingAttributeName](const FSkeletonAttributeBinding& AttributeBinding)
+	{
+		return AttributeBinding.BoneName == BindingAttributeName;
+	});
+}
+
+FSkeletonNamedAttributeMappingEntry USkeletonBinding::CreateDefaultMappingEntry(const FSkeletonNamedAttributeMapping& AttributeMapping, const FSkeletonNamedAttribute& NamedAttribute)
+{
+	FSkeletonNamedAttributeMappingEntry NewMappingEntry;
+	NewMappingEntry.AttributeName = NamedAttribute.Name;
+
+	const FSkeletonNamedAttributeMappingType* const MappingType = AttributeMapping.MappingTypes.FindByPredicate([NamedAttribute](const FSkeletonNamedAttributeMappingType& MappingType)
+	{
+		return MappingType.SourceType == NamedAttribute.Type;
+	});
+
+	if (MappingType)
+	{
+		NewMappingEntry.Value = NewObject<USkeletonTemplateBaseType>(this, MappingType->TargetType);
+	}
+
+	return NewMappingEntry;
+}
+
+bool USkeletonBinding::BindAttribute(const FName BindingAttributeName, const FName TemplateAttributeName)
+{
+	if (TemplateAttributeName == NAME_None || BindingAttributeName == NAME_None)
+	{
+		return false;
+	}
+	
+	const FSkeletonAttributeBinding* const ExistingAttributeBinding = FindAttributeBinding(BindingAttributeName);
+	if (ExistingAttributeBinding != nullptr)
+	{
+		return false; // Template attribute is already bound
+	}
+
+	FSkeletonAttributeBinding NewBinding;
+	NewBinding.AttributeName = TemplateAttributeName;
+	NewBinding.BoneName = BindingAttributeName;
+
+	AttributeBindings.Add(NewBinding);
+	return true;
+}
+
+bool USkeletonBinding::UnbindTemplateNamedAttribute(const FName TemplateAttributeName)
+{
+	const int32 ExistingAttributeBindingIndex = AttributeBindings.IndexOfByPredicate([TemplateAttributeName](const FSkeletonAttributeBinding& AttributeBinding)
+	{
+		return AttributeBinding.AttributeName == TemplateAttributeName;
+	});
+	
+	if (ExistingAttributeBindingIndex == INDEX_NONE)
+	{
+		return false; // Template attribute is not bound
+	}
+
+	AttributeBindings.RemoveAt(ExistingAttributeBindingIndex);	
+	return true;
+}
+
+bool USkeletonBinding::UnbindBindingNamedAttribute(const FName BindingAttributeName)
+{
+	const int32 ExistingAttributeBindingIndex = FindAttributeBindingIndex(BindingAttributeName);
+	if (ExistingAttributeBindingIndex == INDEX_NONE)
+	{
+		return false; // Template attribute is not bound
+	}
+
+	AttributeBindings.RemoveAt(ExistingAttributeBindingIndex);	
+	return true;
+}
+
+void USkeletonBinding::GetChildNamedAttributes(const FName BindingAttributeName, TArray<FName>& OutChildren)
+{
+	for (const FSkeletonNamedAttribute& NamedAttribute : NamedAttributes)
+	{
+		if (NamedAttribute.ParentName == BindingAttributeName)
+		{
+			OutChildren.Add(NamedAttribute.Name);
+			GetChildNamedAttributes(NamedAttribute.Name, OutChildren);
+		}
+	}
+}
+
+bool USkeletonBinding::AddNamedAttributeToSet(const FName AttributeSetName, const FName AttributeName)
+{
+	const FSkeletonNamedAttribute* const NamedAttribute = FindNamedAttribute(AttributeName);
+	if (!NamedAttribute)
+	{
+		return false;
+	}
+	
+	FSkeletonNamedAttributeSet* AttributeSet = FindMutableNamedAttributeSet(AttributeSetName);
+	if (!AttributeSet)
+	{
+		return false;
+	}
+
+	if (AttributeSet->NamedAttributes.Contains(AttributeName))
+	{
+		return false;
+	}
+
+	AttributeSet->NamedAttributes.Add(AttributeName);
+
+	// Add this attribute to all mappings with this set as a source set
+	for (FSkeletonNamedAttributeMapping& AttributeMapping : NamedAttributeMappings)
+	{
+		if (AttributeMapping.SourceAttributeSet == AttributeSetName)
+		{
+			FSkeletonNamedAttributeMappingEntry NewMappingEntry = CreateDefaultMappingEntry(AttributeMapping, *NamedAttribute);
+			AttributeMapping.TableData.Add(NewMappingEntry);
+		}
+	}
+	
+	return true;
+}
+
+bool USkeletonBinding::RemoveNamedAttributeFromSet(const FName AttributeSetName, const FName AttributeName)
+{
+	FSkeletonNamedAttributeSet* AttributeSet = FindMutableNamedAttributeSet(AttributeSetName);
+	if (!AttributeSet)
+	{
+		return false;
+	}
+
+	return AttributeSet->NamedAttributes.Remove(AttributeName) > 0;
+}
+
+FSkeletonNamedAttributeSet* USkeletonBinding::FindMutableNamedAttributeSet(const FName AttributeSetName)
+{
+	return NamedAttributeSets.FindByPredicate([AttributeSetName](const FSkeletonNamedAttributeSet& Candidate)
+		{
+			return Candidate.Name == AttributeSetName;
+		});
+}
+
+const FSkeletonNamedAttributeMapping* const USkeletonBinding::FindNamedAttributeMapping(const FName AttributeMappingName) const
+{
+	return NamedAttributeMappings.FindByPredicate([AttributeMappingName](const FSkeletonNamedAttributeMapping& Candidate)
+		{
+			return Candidate.Name == AttributeMappingName;
+		});
 }
\ No newline at end of file
diff --git a/Engine/Plugins/Experimental/Animation/SkeletonTemplateFramework/Source/Runtime/Public/SkeletonTemplateFramework/SkeletonBinding.h b/Engine/Plugins/Experimental/Animation/SkeletonTemplateFramework/Source/Runtime/Public/SkeletonTemplateFramework/SkeletonBinding.h
--- a/Engine/Plugins/Experimental/Animation/SkeletonTemplateFramework/Source/Runtime/Public/SkeletonTemplateFramework/SkeletonBinding.h
+++ b/Engine/Plugins/Experimental/Animation/SkeletonTemplateFramework/Source/Runtime/Public/SkeletonTemplateFramework/SkeletonBinding.h
@@ -37,8 +37,50 @@ public:
 
 	// Sets up the named attributes using the provided template
 	SKELETONTEMPLATEFRAMEWORKRUNTIME_API bool InitializeFrom(const TObjectPtr<USkeletonTemplate> InTemplate, const TObjectPtr<USkeleton> InSkeleton);
+
+	SKELETONTEMPLATEFRAMEWORKRUNTIME_API const FSkeletonNamedAttribute* const FindNamedAttribute(const FName AttributeName) const;
+	
+	SKELETONTEMPLATEFRAMEWORKRUNTIME_API TObjectPtr<USkeleton> GetSkeleton() const;
+	
+	SKELETONTEMPLATEFRAMEWORKRUNTIME_API TObjectPtr<USkeletonTemplate> GetSkeletonTemplate() const;
+	
+	SKELETONTEMPLATEFRAMEWORKRUNTIME_API const TArray<FSkeletonNamedAttribute>& GetNamedAttributes() const;
+
+	SKELETONTEMPLATEFRAMEWORKRUNTIME_API const TArray<FSkeletonNamedAttributeSet>& GetNamedAttributeSets() const;
+
+	SKELETONTEMPLATEFRAMEWORKRUNTIME_API const TArray<FSkeletonNamedAttributeMapping>& GetNamedAttributeMappings() const;
+
+	// Returns the array containing all schema named attributes that are not currently bound
+	SKELETONTEMPLATEFRAMEWORKRUNTIME_API TArray<FSkeletonNamedAttribute> GetUnboundSchemaNamedAttributes() const;
+
+	SKELETONTEMPLATEFRAMEWORKRUNTIME_API const FSkeletonAttributeBinding* const FindAttributeBinding(const FName BindingAttributeName) const;
+	
+	SKELETONTEMPLATEFRAMEWORKRUNTIME_API bool BindAttribute(const FName BindingAttributeName, const FName TemplateAttributeName);
+
+	SKELETONTEMPLATEFRAMEWORKRUNTIME_API bool UnbindTemplateNamedAttribute(const FName TemplateAttributeName);
+
+	SKELETONTEMPLATEFRAMEWORKRUNTIME_API bool UnbindBindingNamedAttribute(const FName BindingAttributeName);
+
+	SKELETONTEMPLATEFRAMEWORKRUNTIME_API void GetChildNamedAttributes(const FName BindingAttributeName, TArray<FName>& OutChildren);
+
+	// Sets
+
+	SKELETONTEMPLATEFRAMEWORKRUNTIME_API bool AddNamedAttributeToSet(const FName AttributeSetName, const FName AttributeName);
+
+	SKELETONTEMPLATEFRAMEWORKRUNTIME_API bool RemoveNamedAttributeFromSet(const FName AttributeSetName, const FName AttributeName);
+
+	// Mappings
+
+	SKELETONTEMPLATEFRAMEWORKRUNTIME_API const FSkeletonNamedAttributeMapping* const FindNamedAttributeMapping(const FName AttributeMappingName) const;
 	
 protected:
+
+	FSkeletonNamedAttributeSet* FindMutableNamedAttributeSet(const FName AttributeSetName);
+	
+	int32 FindAttributeBindingIndex(const FName BindingAttributeName) const;
+
+	FSkeletonNamedAttributeMappingEntry CreateDefaultMappingEntry(const FSkeletonNamedAttributeMapping& AttributeMapping, const FSkeletonNamedAttribute& NamedAttribute);
+
 	// The skeleton template to bind
 	// We inherit everything from our parent template, allowing us to override things as needed
 	UPROPERTY(EditAnywhere, Category = "Settings")
@@ -54,7 +96,7 @@ protected:
 	UPROPERTY(EditAnywhere, Category = "Settings")
 	TArray<FSkeletonAttributeBinding> AttributeBindings;
 
-	// The list of attributes within this binding (includes attributes from the template and any unmapped entries from the skeleton)
+	// The list of attributes within this binding, includes all entries from the skeleton
 	UPROPERTY(VisibleAnywhere, Category = "Settings")
 	TArray<FSkeletonNamedAttribute> NamedAttributes;
 
