Skeleton Template Editor - Fix LOCTEXT conflict

diff --git a/Engine/Plugins/Experimental/Animation/SkeletonTemplateFramework/Source/Editor/Private/SkeletonTemplateFramework/SkeletonTemplateEditorToolkit.cpp b/Engine/Plugins/Experimental/Animation/SkeletonTemplateFramework/Source/Editor/Private/SkeletonTemplateFramework/SkeletonTemplateEditorToolkit.cpp
--- a/Engine/Plugins/Experimental/Animation/SkeletonTemplateFramework/Source/Editor/Private/SkeletonTemplateFramework/SkeletonTemplateEditorToolkit.cpp
+++ b/Engine/Plugins/Experimental/Animation/SkeletonTemplateFramework/Source/Editor/Private/SkeletonTemplateFramework/SkeletonTemplateEditorToolkit.cpp
@@ -97,7 +97,7 @@ void FSkeletonTemplateEditorToolkit::RegisterTabSpawners(const TSharedRef<class
 	
 	InTabManager->RegisterTabSpawner(FEditorTabs::AttributeMappingsId, FOnSpawnTab::CreateSP(this, &FSkeletonTemplateEditorToolkit::SpawnTabAttributeMappings))
 		.SetDisplayName(LOCTEXT("AttributeMappingsTabMenu_Description", "Attribute Mappings"))
-		.SetTooltipText(LOCTEXT("AttributeMappingsTabMenu_Description", "Shows the attribute mappings view panel"))
+		.SetTooltipText(LOCTEXT("AttributeMappingsTabMenu_ToolTip", "Shows the attribute mappings view panel"))
 		.SetGroup(WorkspaceMenuCategoryRef)
 		.SetIcon(FSlateIcon(FAppStyle::GetAppStyleSetName(), "LevelEditor.Tabs.Outliner"));
 
