Define ShortNames for Skeleton Template Framework modules
* Builds can fail due to file paths exceeding max size limit

diff --git a/Engine/Plugins/Experimental/Animation/SkeletonTemplateFramework/Source/Editor/SkeletonTemplateFrameworkEditor.Build.cs b/Engine/Plugins/Experimental/Animation/SkeletonTemplateFramework/Source/Editor/SkeletonTemplateFrameworkEditor.Build.cs
--- a/Engine/Plugins/Experimental/Animation/SkeletonTemplateFramework/Source/Editor/SkeletonTemplateFrameworkEditor.Build.cs
+++ b/Engine/Plugins/Experimental/Animation/SkeletonTemplateFramework/Source/Editor/SkeletonTemplateFrameworkEditor.Build.cs
@@ -7,6 +7,8 @@ public class SkeletonTemplateFrameworkEditor : ModuleRules
 {
 	public SkeletonTemplateFrameworkEditor(ReadOnlyTargetRules Target) : base(Target)
 	{
+		ShortName = "STFEd";
+
 		PrivateDependencyModuleNames.AddAll(
 			"AssetDefinition",
 			"Core",
diff --git a/Engine/Plugins/Experimental/Animation/SkeletonTemplateFramework/Source/Runtime/SkeletonTemplateFrameworkRuntime.Build.cs b/Engine/Plugins/Experimental/Animation/SkeletonTemplateFramework/Source/Runtime/SkeletonTemplateFrameworkRuntime.Build.cs
--- a/Engine/Plugins/Experimental/Animation/SkeletonTemplateFramework/Source/Runtime/SkeletonTemplateFrameworkRuntime.Build.cs
+++ b/Engine/Plugins/Experimental/Animation/SkeletonTemplateFramework/Source/Runtime/SkeletonTemplateFrameworkRuntime.Build.cs
@@ -6,6 +6,8 @@ public class SkeletonTemplateFrameworkRuntime : ModuleRules
 {
 	public SkeletonTemplateFrameworkRuntime(ReadOnlyTargetRules Target) : base(Target)
 	{
+		ShortName = "STFRun";
+
 		PublicDependencyModuleNames.AddRange(
 			new string[]
 			{
