Abstract Skeleton Runtime

diff --git a/Engine/Plugins/Experimental/Animation/SkeletonTemplateFramework/Source/Runtime/Private/SkeletonTemplateFramework/AbstractSkeletonLabelBinding.cpp b/Engine/Plugins/Experimental/Animation/SkeletonTemplateFramework/Source/Runtime/Private/SkeletonTemplateFramework/AbstractSkeletonLabelBinding.cpp
new file mode 100644
--- /dev/null
+++ b/Engine/Plugins/Experimental/Animation/SkeletonTemplateFramework/Source/Runtime/Private/SkeletonTemplateFramework/AbstractSkeletonLabelBinding.cpp
@@ -0,0 +1,112 @@
+// Copyright Epic Games, Inc. All Rights Reserved.
+
+#include "SkeletonTemplateFramework/AbstractSkeletonLabelBinding.h"
+
+#if WITH_EDITOR
+
+bool UAbstractSkeletonLabelBinding::BindBoneToLabel(const TObjectPtr<const UAbstractSkeletonLabelCollection> InCollection, const FName BoneName, const FName Label)
+{
+	if (!IsLabelBound(InCollection, BoneName))
+	{
+		FAbstractSkeleton_LabelBinding& NewBinding = LabelBindings.AddDefaulted_GetRef();
+		NewBinding.BoneName = BoneName;
+		NewBinding.Label = Label;
+		NewBinding.SourceCollection = InCollection;
+
+		return true;
+	}
+	else
+	{
+		return false;
+	}
+}
+
+bool UAbstractSkeletonLabelBinding::UnbindBoneFromLabel(const TObjectPtr<const UAbstractSkeletonLabelCollection> InCollection, const FName BoneName, const FName Label)
+{
+	if (IsLabelBound(InCollection, Label))
+	{
+		const int32 BindingIndex = LabelBindings.IndexOfByPredicate([&](const FAbstractSkeleton_LabelBinding& Binding)
+		{
+			return Binding.SourceCollection == InCollection && Binding.BoneName == BoneName && Binding.Label == Label;
+		});
+
+		if (ensure(BindingIndex != INDEX_NONE))
+		{
+			LabelBindings.RemoveAt(BindingIndex);
+			return true;
+		}
+	}
+
+	return false;
+}
+
+bool UAbstractSkeletonLabelBinding::AddLabelCollection(const TObjectPtr<const UAbstractSkeletonLabelCollection> InCollection)
+{
+	if (!LabelCollections.Contains(InCollection))
+	{
+		LabelCollections.Add(InCollection);
+		return true;
+	}
+
+	return false;
+}
+
+bool UAbstractSkeletonLabelBinding::RemoveLabelCollection(const TObjectPtr<const UAbstractSkeletonLabelCollection> InCollection)
+{
+	LabelCollections.Remove(InCollection);
+
+	LabelBindings.RemoveAll([InCollection](const FAbstractSkeleton_LabelBinding& InBinding)
+	{
+		return InBinding.SourceCollection == InCollection;
+	});
+
+	return true;
+}
+
+bool UAbstractSkeletonLabelBinding::SetSkeleton(const TObjectPtr<USkeleton> InSkeleton)
+{
+	if (InSkeleton != Skeleton)
+	{
+		Skeleton = InSkeleton;
+		LabelBindings.Empty();
+		
+		return true;
+	}
+
+	return false;
+}
+
+#endif // WITH_EDITOR
+
+bool UAbstractSkeletonLabelBinding::IsLabelBound(const TObjectPtr<const UAbstractSkeletonLabelCollection> InCollection, const FName LabelName) const
+{
+	return LabelBindings.ContainsByPredicate([&](const FAbstractSkeleton_LabelBinding& Binding)
+	{
+		return Binding.SourceCollection == InCollection && Binding.Label == LabelName;
+	});
+}
+
+FName UAbstractSkeletonLabelBinding::GetLabelBinding(const TObjectPtr<const UAbstractSkeletonLabelCollection> InCollection, const FName LabelName) const
+{
+	const int32 BindingIndex = LabelBindings.IndexOfByPredicate([&](const FAbstractSkeleton_LabelBinding& Binding)
+		{
+			return Binding.SourceCollection == InCollection && Binding.Label == LabelName;
+		});
+	
+	return BindingIndex == INDEX_NONE ? NAME_None : LabelBindings[BindingIndex].BoneName;
+}
+
+TConstArrayView<FAbstractSkeleton_LabelBinding> UAbstractSkeletonLabelBinding::GetLabelBindings() const
+{
+	return LabelBindings;
+}
+
+TObjectPtr<USkeleton> UAbstractSkeletonLabelBinding::GetSkeleton() const
+{
+	return Skeleton;
+}
+
+TConstArrayView<TObjectPtr<const UAbstractSkeletonLabelCollection>> UAbstractSkeletonLabelBinding::GetLabelCollections() const
+{
+	return LabelCollections;
+}
\ No newline at end of file
diff --git a/Engine/Plugins/Experimental/Animation/SkeletonTemplateFramework/Source/Runtime/Private/SkeletonTemplateFramework/AbstractSkeletonLabelCollection.cpp b/Engine/Plugins/Experimental/Animation/SkeletonTemplateFramework/Source/Runtime/Private/SkeletonTemplateFramework/AbstractSkeletonLabelCollection.cpp
new file mode 100644
--- /dev/null
+++ b/Engine/Plugins/Experimental/Animation/SkeletonTemplateFramework/Source/Runtime/Private/SkeletonTemplateFramework/AbstractSkeletonLabelCollection.cpp
@@ -0,0 +1,66 @@
+// Copyright Epic Games, Inc. All Rights Reserved.
+
+#include "SkeletonTemplateFramework/AbstractSkeletonLabelCollection.h"
+
+#if WITH_EDITOR
+
+bool UAbstractSkeletonLabelCollection::AddLabel(const FName InLabel)
+{
+	if (!HasLabel(InLabel))
+	{
+		Labels.Add(InLabel);
+		return true;
+	}
+	else
+	{
+		return false;
+	}
+}
+
+bool UAbstractSkeletonLabelCollection::RenameLabel(const FName InOldLabel, const FName InNewLabel)
+{
+	const int32 ExistingLabelIndex = Labels.Find(InOldLabel);
+	
+	if (ExistingLabelIndex == INDEX_NONE)
+	{
+		return false;
+	}
+
+	if (HasLabel(InNewLabel))
+	{
+		return false;
+	}
+
+	Labels[ExistingLabelIndex] = InNewLabel;
+	return true;
+}
+
+bool UAbstractSkeletonLabelCollection::RemoveLabel(const FName InLabel)
+{
+	if (HasLabel(InLabel))
+	{
+		Labels.Remove(InLabel);
+		return true;
+	}
+	else
+	{
+		return false;
+	}
+}
+
+TArray<FName>& UAbstractSkeletonLabelCollection::GetMutableLabels()
+{
+	return Labels;
+}
+
+#endif // WITH_EDITOR
+
+bool UAbstractSkeletonLabelCollection::HasLabel(const FName InLabel) const
+{
+	return Labels.Contains(InLabel);
+}
+
+TConstArrayView<FName> UAbstractSkeletonLabelCollection::GetLabels() const
+{
+	return Labels;
+}
\ No newline at end of file
diff --git a/Engine/Plugins/Experimental/Animation/SkeletonTemplateFramework/Source/Runtime/Private/SkeletonTemplateFramework/AbstractSkeletonSetBinding.cpp b/Engine/Plugins/Experimental/Animation/SkeletonTemplateFramework/Source/Runtime/Private/SkeletonTemplateFramework/AbstractSkeletonSetBinding.cpp
new file mode 100644
--- /dev/null
+++ b/Engine/Plugins/Experimental/Animation/SkeletonTemplateFramework/Source/Runtime/Private/SkeletonTemplateFramework/AbstractSkeletonSetBinding.cpp
@@ -0,0 +1,183 @@
+// Copyright Epic Games, Inc. All Rights Reserved.
+
+#include "SkeletonTemplateFramework/AbstractSkeletonSetBinding.h"
+
+#include "SkeletonTemplateFramework/AbstractSkeletonSetCollection.h"
+
+// Bone Bindings
+
+bool UAbstractSkeletonSetBinding::AddBoneToSet(const FName BoneName, const FName SetName)
+{
+	if (!SetCollection->HasSet(SetName) && SetName != NAME_None)
+	{
+		return false; // Set does not exist
+	}
+
+	if (IsBoneInSet(BoneName))
+	{
+		return false; // Already exists in a set
+	}
+
+	FAbstractSkeleton_BoneBinding& Binding = BoneBindings.AddDefaulted_GetRef();
+	Binding.BoneName = BoneName;
+	Binding.SetName = SetName;
+
+	return true;
+}
+
+bool UAbstractSkeletonSetBinding::RemoveBoneFromSet(const FName BoneName)
+{
+	const int32 BindingIndex = BoneBindings.IndexOfByPredicate([&](const FAbstractSkeleton_BoneBinding& Binding)
+	{
+		return Binding.BoneName == BoneName;
+	});
+
+	if (BindingIndex != INDEX_NONE)
+	{
+		BoneBindings.RemoveAtSwap(BindingIndex);
+		return true;
+	}
+	else
+	{
+		return false;
+	}
+}
+
+bool UAbstractSkeletonSetBinding::IsBoneInSet(const FName BoneName) const
+{
+	return BoneBindings.ContainsByPredicate([&](const FAbstractSkeleton_BoneBinding& Binding)
+		{
+			return Binding.BoneName == BoneName;
+		});
+}
+
+bool UAbstractSkeletonSetBinding::IsBoneInSet(const FName BoneName, const FName SetName) const
+{
+	return BoneBindings.ContainsByPredicate([&](const FAbstractSkeleton_BoneBinding& Binding)
+		{
+			return Binding.BoneName == BoneName && Binding.SetName == SetName;
+		});
+}
+
+FName UAbstractSkeletonSetBinding::GetBoneSet(const FName BoneName)
+{
+	const FAbstractSkeleton_BoneBinding* const Binding = BoneBindings.FindByPredicate([&](const FAbstractSkeleton_BoneBinding& Binding)
+		{
+			return Binding.BoneName == BoneName;
+		});
+
+	return Binding ? Binding->SetName : NAME_None;
+}
+
+const TConstArrayView<FAbstractSkeleton_BoneBinding> UAbstractSkeletonSetBinding::GetBoneBindings() const
+{
+	return BoneBindings;
+}
+
+// Attribute Bindings
+
+bool UAbstractSkeletonSetBinding::AddAttributeToSet(const FAnimationAttributeIdentifier Attribute, const FName SetName)
+{
+	if (!SetCollection->HasSet(SetName) && SetName != NAME_None)
+	{
+		return false; // Set does not exist
+	}
+
+	if (IsAttributeInSet(Attribute))
+	{
+		return false; // Already exists in a set
+	}
+
+	if (IsAttributeInSet(Attribute, NAME_None))
+	{
+		FAbstractSkeleton_AttributeBinding* UnboundAttribute = AttributeBindings.FindByPredicate([Attribute, SetName](const FAbstractSkeleton_AttributeBinding& Binding)
+		{
+			return Binding.Attribute == Attribute && Binding.SetName == NAME_None;
+		});
+
+		if (ensure(UnboundAttribute))
+		{
+			UnboundAttribute->SetName = SetName;
+			return true;
+		}
+		else
+		{
+			return false;
+		}
+	}
+
+	FAbstractSkeleton_AttributeBinding& Binding = AttributeBindings.AddDefaulted_GetRef();
+	Binding.Attribute = Attribute;
+	Binding.SetName = SetName;
+
+	return true;
+}
+
+bool UAbstractSkeletonSetBinding::RemoveAttributeFromSet(const FAnimationAttributeIdentifier Attribute)
+{
+	const int32 BindingIndex = AttributeBindings.IndexOfByPredicate([&](const FAbstractSkeleton_AttributeBinding& Binding)
+	{
+		return Binding.Attribute == Attribute;
+	});
+
+	if (BindingIndex != INDEX_NONE)
+	{
+		AttributeBindings.RemoveAtSwap(BindingIndex);
+		return true;
+	}
+	else
+	{
+		return false;
+	}
+}
+
+bool UAbstractSkeletonSetBinding::IsAttributeInSet(const FAnimationAttributeIdentifier Attribute) const
+{
+	return AttributeBindings.ContainsByPredicate([&](const FAbstractSkeleton_AttributeBinding& Binding)
+		{
+			return Binding.Attribute.GetName() == Attribute.GetName() && Binding.Attribute.GetBoneName() == Attribute.GetBoneName() && Binding.SetName != NAME_None;
+		});
+}
+
+bool UAbstractSkeletonSetBinding::IsAttributeInSet(const FAnimationAttributeIdentifier Attribute, const FName SetName) const
+{
+	return AttributeBindings.ContainsByPredicate([&](const FAbstractSkeleton_AttributeBinding& Binding)
+		{
+			return Binding.Attribute.GetName() == Attribute.GetName() && Binding.Attribute.GetBoneName() == Attribute.GetBoneName() && Binding.SetName == SetName;
+		});
+}
+
+const TConstArrayView<FAbstractSkeleton_AttributeBinding> UAbstractSkeletonSetBinding::GetAttributeBindings() const
+{
+	return AttributeBindings;
+}
+
+bool UAbstractSkeletonSetBinding::SetSkeleton(const TObjectPtr<USkeleton> InSkeleton)
+{
+	Skeleton = InSkeleton;
+
+	BoneBindings.Empty();
+	AttributeBindings.Empty();
+
+	return true;
+}
+
+TObjectPtr<USkeleton> UAbstractSkeletonSetBinding::GetSkeleton() const
+{
+	return Skeleton;
+}
+
+bool UAbstractSkeletonSetBinding::SetSetCollection(const TObjectPtr<const UAbstractSkeletonSetCollection> InSetCollection)
+{
+	SetCollection = InSetCollection;
+
+	BoneBindings.Empty();
+	AttributeBindings.Empty();
+
+	return true;
+}
+
+TObjectPtr<const UAbstractSkeletonSetCollection> UAbstractSkeletonSetBinding::GetSetCollection() const
+{
+	return SetCollection;
+}
diff --git a/Engine/Plugins/Experimental/Animation/SkeletonTemplateFramework/Source/Runtime/Private/SkeletonTemplateFramework/AbstractSkeletonSetCollection.cpp b/Engine/Plugins/Experimental/Animation/SkeletonTemplateFramework/Source/Runtime/Private/SkeletonTemplateFramework/AbstractSkeletonSetCollection.cpp
new file mode 100644
--- /dev/null
+++ b/Engine/Plugins/Experimental/Animation/SkeletonTemplateFramework/Source/Runtime/Private/SkeletonTemplateFramework/AbstractSkeletonSetCollection.cpp
@@ -0,0 +1,161 @@
+// Copyright Epic Games, Inc. All Rights Reserved.
+
+#include "SkeletonTemplateFramework/AbstractSkeletonSetCollection.h"
+#include "SkeletonTemplateFramework/Utils.h"
+
+bool UAbstractSkeletonSetCollection::AddSet(const FName InSetName, const FName InParentSetName)
+{
+	if (HasSet(InSetName))
+	{
+		return false; // Set with this name already exists
+	}
+
+	FAbstractSkeletonSet& NewSet = SetHierarchy.AddDefaulted_GetRef();
+	NewSet.SetName = InSetName;
+	NewSet.ParentSetName = InParentSetName;
+
+	SortSetHierarchy();
+
+	return true;
+}
+
+bool UAbstractSkeletonSetCollection::HasSet(const FName InSetName) const
+{
+	return GetSetIndex(InSetName) != INDEX_NONE;
+}
+
+bool UAbstractSkeletonSetCollection::RenameSet(const FName OldSetName, const FName NewSetName)
+{
+	if (HasSet(NewSetName))
+	{
+		return false; // Set with this name already exists
+	}
+
+	for (FAbstractSkeletonSet& Set : SetHierarchy)
+	{
+		if (Set.SetName == OldSetName)
+		{
+			Set.SetName = NewSetName;
+		}
+		else if (Set.ParentSetName == OldSetName)
+		{
+			Set.ParentSetName = NewSetName;
+		}
+	}
+
+	return true;
+}
+
+bool UAbstractSkeletonSetCollection::RemoveSet(const FName SetName)
+{
+	if (SetName == NAME_None )
+	{
+		return false;
+	}
+
+	const int32 SetIndex = GetSetIndex(SetName);
+	if (SetIndex == INDEX_NONE)
+	{
+		return false;
+	}
+
+	const FName SetToRemoveParentSetName = GetParentSetName(SetName);
+
+	// Reparent all children of the set we're remove
+	for (FAbstractSkeletonSet& Set : SetHierarchy)
+	{
+		if (Set.ParentSetName == SetName)
+		{
+			Set.ParentSetName = SetToRemoveParentSetName;
+		}
+	}
+
+	SetHierarchy.RemoveAt(SetIndex);
+
+	return true;
+}
+
+bool UAbstractSkeletonSetCollection::ReparentSet(const FName SetName, const FName NewParentName)
+{
+	if (SetName == NewParentName || IsDescendantOf(NewParentName, SetName))
+	{
+		return false;
+	}
+
+	const int32 SetIndex = GetSetIndex(SetName);
+	if (SetIndex != INDEX_NONE)
+	{
+		SetHierarchy[SetIndex].ParentSetName = NewParentName;
+		SortSetHierarchy();
+
+		return true;
+	}
+
+	return false;
+}
+
+FName UAbstractSkeletonSetCollection::GetParentSetName(const FName SetName) const
+{
+	for (const FAbstractSkeletonSet& Set : SetHierarchy)
+	{
+		if (Set.SetName == SetName)
+		{
+			return Set.ParentSetName;
+		}
+	}
+
+	return NAME_None;
+}
+
+const TArrayView<const FAbstractSkeletonSet> UAbstractSkeletonSetCollection::GetSetHierarchy() const
+{
+	return SetHierarchy;
+}
+
+int32 UAbstractSkeletonSetCollection::GetSetIndex(const FName SetName) const
+{
+	for (int32 SetIndex = 0; SetIndex < SetHierarchy.Num(); ++SetIndex)
+	{
+		const FAbstractSkeletonSet Set = SetHierarchy[SetIndex];
+
+		if (Set.SetName == SetName)
+		{
+			return SetIndex;
+		}
+	}
+
+	return INDEX_NONE;
+}
+
+bool UAbstractSkeletonSetCollection::IsDescendantOf(const FName SetName, const FName AncestorSetName) const
+{
+	FName ParentSetName = GetParentSetName(SetName);
+	
+	while (ParentSetName != NAME_None)
+	{
+		if (ParentSetName == AncestorSetName)
+		{
+			return true;
+		}
+		ParentSetName = GetParentSetName(ParentSetName);
+	}
+
+	return false;
+}
+
+void UAbstractSkeletonSetCollection::SortSetHierarchy()
+{
+	TFunction<FAbstractSkeletonSet* (FName)> Functor = [this](FName InAttributeName)
+		{
+			return &GetSetRef(InAttributeName);
+		};
+
+	UE::Anim::SkelTemplates::Util::SortByProperty<FAbstractSkeletonSet>(SetHierarchy, &FAbstractSkeletonSet::SetName, &FAbstractSkeletonSet::ParentSetName, Functor);
+}
+
+FAbstractSkeletonSet& UAbstractSkeletonSetCollection::GetSetRef(const FName SetName)
+{
+	const int32 SetIndex = GetSetIndex(SetName);
+	check(SetIndex != INDEX_NONE)
+	return SetHierarchy[SetIndex];
+}
\ No newline at end of file
diff --git a/Engine/Plugins/Experimental/Animation/SkeletonTemplateFramework/Source/Runtime/Private/SkeletonTemplateFramework/Module.cpp b/Engine/Plugins/Experimental/Animation/SkeletonTemplateFramework/Source/Runtime/Private/SkeletonTemplateFramework/Module.cpp
--- a/Engine/Plugins/Experimental/Animation/SkeletonTemplateFramework/Source/Runtime/Private/SkeletonTemplateFramework/Module.cpp
+++ b/Engine/Plugins/Experimental/Animation/SkeletonTemplateFramework/Source/Runtime/Private/SkeletonTemplateFramework/Module.cpp
@@ -2,7 +2,7 @@
 
 #include "Modules/ModuleManager.h"
 
-namespace UE::Anim::STF
+namespace UE::Anim::SkelTemplates
 {
 	class FRuntimeModule : public IModuleInterface
 	{
@@ -17,4 +17,4 @@ namespace UE::Anim::STF
 	};
 }
 
-IMPLEMENT_MODULE(UE::Anim::STF::FRuntimeModule, SkeletonTemplateFrameworkRuntime)
+IMPLEMENT_MODULE(UE::Anim::SkelTemplates::FRuntimeModule, SkeletonTemplateFrameworkRuntime)
diff --git a/Engine/Plugins/Experimental/Animation/SkeletonTemplateFramework/Source/Runtime/Private/SkeletonTemplateFramework/Utils.h b/Engine/Plugins/Experimental/Animation/SkeletonTemplateFramework/Source/Runtime/Private/SkeletonTemplateFramework/Utils.h
new file mode 100644
--- /dev/null
+++ b/Engine/Plugins/Experimental/Animation/SkeletonTemplateFramework/Source/Runtime/Private/SkeletonTemplateFramework/Utils.h
@@ -0,0 +1,75 @@
+#pragma once
+
+#include "CoreMinimal.h"
+#include "Containers/Queue.h"
+
+namespace UE::Anim::SkelTemplates::Util
+{
+	template <typename ArrayType>
+	void SortByProperty(TArray<ArrayType>& Array, FName ArrayType::* NameProperty, FName ArrayType::* ParentNameProperty, TFunctionRef<ArrayType* (FName)> FindArrayItem)
+	{
+		struct Node
+		{
+			FName AttributeName;
+			TArray<FName> Children;
+			int32 InDegree;
+
+			Node(const FName InAttributeName)
+				: AttributeName(InAttributeName)
+				, InDegree(0)
+			{
+			}
+		};
+
+		TMap<FName, Node> Graph;
+
+		for (const ArrayType& ArrayItem : Array)
+		{
+			Graph.Add(ArrayItem.*NameProperty, Node(ArrayItem.*NameProperty));
+		}
+
+		for (const ArrayType& ArrayItem : Array)
+		{
+			if (ArrayItem.*ParentNameProperty != NAME_None)
+			{
+				Graph[ArrayItem.*ParentNameProperty].Children.Add(ArrayItem.*NameProperty);
+				++Graph[ArrayItem.*NameProperty].InDegree;
+			}
+		}
+
+		TQueue<FName> NodeQueue;
+		for (const ArrayType& Item : Array)
+		{
+			const Node& ItemNode = Graph[Item.*NameProperty];
+			if (ItemNode.InDegree == 0)
+			{
+				NodeQueue.Enqueue(ItemNode.AttributeName);
+			}
+		}
+
+		TArray<ArrayType> SortedArray;
+
+		while (!NodeQueue.IsEmpty())
+		{
+			FName AttributeName;
+			NodeQueue.Dequeue(AttributeName);
+
+			if (AttributeName != NAME_None)
+			{
+				SortedArray.Add(*FindArrayItem(AttributeName));
+			}
+
+			for (const FName& Child : Graph[AttributeName].Children)
+			{
+				--Graph[Child].InDegree;
+				if (Graph[Child].InDegree == 0)
+				{
+					NodeQueue.Enqueue(Child);
+				}
+			}
+		}
+
+		ensureMsgf(Array.Num() == SortedArray.Num(), TEXT("Cycle found in the DAG"));
+		Array = SortedArray;
+	}
+}
\ No newline at end of file
diff --git a/Engine/Plugins/Experimental/Animation/SkeletonTemplateFramework/Source/Runtime/Public/SkeletonTemplateFramework/AbstractSkeletonLabelBinding.h b/Engine/Plugins/Experimental/Animation/SkeletonTemplateFramework/Source/Runtime/Public/SkeletonTemplateFramework/AbstractSkeletonLabelBinding.h
new file mode 100644
--- /dev/null
+++ b/Engine/Plugins/Experimental/Animation/SkeletonTemplateFramework/Source/Runtime/Public/SkeletonTemplateFramework/AbstractSkeletonLabelBinding.h
@@ -0,0 +1,65 @@
+// Copyright Epic Games, Inc. All Rights Reserved.
+
+#pragma once
+
+#include "CoreMinimal.h"
+
+#include "AbstractSkeletonLabelBinding.generated.h"
+
+class UAbstractSkeletonLabelCollection;
+class USkeleton;
+
+USTRUCT()
+struct FAbstractSkeleton_LabelBinding
+{
+	GENERATED_BODY()
+
+	UPROPERTY()
+	FName Label;
+
+	UPROPERTY()
+	FName BoneName;
+
+	UPROPERTY()
+	TObjectPtr<const UAbstractSkeletonLabelCollection> SourceCollection;
+};
+
+UCLASS(MinimalAPI, BlueprintType)
+class UAbstractSkeletonLabelBinding : public UObject
+{
+	GENERATED_BODY()
+
+public:
+#if WITH_EDITOR
+	SKELETONTEMPLATEFRAMEWORKRUNTIME_API bool BindBoneToLabel(const TObjectPtr<const UAbstractSkeletonLabelCollection> InCollection, const FName BoneName, const FName Label);
+
+	SKELETONTEMPLATEFRAMEWORKRUNTIME_API bool UnbindBoneFromLabel(const TObjectPtr<const UAbstractSkeletonLabelCollection> InCollection, const FName BoneName, const FName Label);
+
+	SKELETONTEMPLATEFRAMEWORKRUNTIME_API bool AddLabelCollection(const TObjectPtr<const UAbstractSkeletonLabelCollection> InCollection);
+
+	SKELETONTEMPLATEFRAMEWORKRUNTIME_API bool RemoveLabelCollection(const TObjectPtr<const UAbstractSkeletonLabelCollection> InCollection);
+	
+	SKELETONTEMPLATEFRAMEWORKRUNTIME_API bool SetSkeleton(const TObjectPtr<USkeleton> InSkeleton);
+
+#endif // WITH_EDITOR
+
+	SKELETONTEMPLATEFRAMEWORKRUNTIME_API bool IsLabelBound(const TObjectPtr<const UAbstractSkeletonLabelCollection> InCollection, const FName LabelName) const;
+
+	SKELETONTEMPLATEFRAMEWORKRUNTIME_API FName GetLabelBinding(const TObjectPtr<const UAbstractSkeletonLabelCollection> InCollection, const FName LabelName) const;
+
+	SKELETONTEMPLATEFRAMEWORKRUNTIME_API TConstArrayView<FAbstractSkeleton_LabelBinding> GetLabelBindings() const;
+	
+	SKELETONTEMPLATEFRAMEWORKRUNTIME_API TObjectPtr<USkeleton> GetSkeleton() const;
+
+	SKELETONTEMPLATEFRAMEWORKRUNTIME_API TConstArrayView<TObjectPtr<const UAbstractSkeletonLabelCollection>> GetLabelCollections() const;
+
+private:
+	UPROPERTY()
+	TArray<TObjectPtr<const UAbstractSkeletonLabelCollection>> LabelCollections;
+
+	UPROPERTY(AssetRegistrySearchable)
+	TObjectPtr<USkeleton> Skeleton;
+
+	UPROPERTY()
+	TArray<FAbstractSkeleton_LabelBinding> LabelBindings;
+};
\ No newline at end of file
diff --git a/Engine/Plugins/Experimental/Animation/SkeletonTemplateFramework/Source/Runtime/Public/SkeletonTemplateFramework/AbstractSkeletonLabelCollection.h b/Engine/Plugins/Experimental/Animation/SkeletonTemplateFramework/Source/Runtime/Public/SkeletonTemplateFramework/AbstractSkeletonLabelCollection.h
new file mode 100644
--- /dev/null
+++ b/Engine/Plugins/Experimental/Animation/SkeletonTemplateFramework/Source/Runtime/Public/SkeletonTemplateFramework/AbstractSkeletonLabelCollection.h
@@ -0,0 +1,32 @@
+// Copyright Epic Games, Inc. All Rights Reserved.
+
+#pragma once
+
+#include "CoreMinimal.h"
+
+#include "AbstractSkeletonLabelCollection.generated.h"
+
+UCLASS(MinimalAPI, BlueprintType)
+class UAbstractSkeletonLabelCollection : public UObject
+{
+	GENERATED_BODY()
+
+public:
+#if WITH_EDITOR
+	SKELETONTEMPLATEFRAMEWORKRUNTIME_API bool AddLabel(const FName InLabel);
+
+	SKELETONTEMPLATEFRAMEWORKRUNTIME_API bool RenameLabel(const FName InOldLabel, const FName InNewLabel);
+
+	SKELETONTEMPLATEFRAMEWORKRUNTIME_API bool RemoveLabel(const FName InLabel);
+
+	SKELETONTEMPLATEFRAMEWORKRUNTIME_API TArray<FName>& GetMutableLabels();
+#endif // WITH_EDITOR
+
+	SKELETONTEMPLATEFRAMEWORKRUNTIME_API bool HasLabel(const FName InLabel) const;
+
+	SKELETONTEMPLATEFRAMEWORKRUNTIME_API TConstArrayView<FName> GetLabels() const;
+
+private:
+	UPROPERTY()
+	TArray<FName> Labels; 
+};
diff --git a/Engine/Plugins/Experimental/Animation/SkeletonTemplateFramework/Source/Runtime/Public/SkeletonTemplateFramework/AbstractSkeletonSetBinding.h b/Engine/Plugins/Experimental/Animation/SkeletonTemplateFramework/Source/Runtime/Public/SkeletonTemplateFramework/AbstractSkeletonSetBinding.h
new file mode 100644
--- /dev/null
+++ b/Engine/Plugins/Experimental/Animation/SkeletonTemplateFramework/Source/Runtime/Public/SkeletonTemplateFramework/AbstractSkeletonSetBinding.h
@@ -0,0 +1,93 @@
+// Copyright Epic Games, Inc. All Rights Reserved.
+
+#pragma once
+
+#include "CoreMinimal.h"
+#include "Animation/AnimData/AttributeIdentifier.h"
+
+#include "AbstractSkeletonSetBinding.generated.h"
+
+class UAbstractSkeletonSetCollection;
+class USkeleton;
+
+USTRUCT()
+struct FAbstractSkeleton_BoneBinding
+{
+	GENERATED_BODY()
+
+	UPROPERTY(EditAnywhere, Category = Default)
+	FName SetName;
+
+	UPROPERTY(EditAnywhere, Category = Default)
+	FName BoneName;
+};
+
+// Curves are represented as float attributes detached from any bone
+USTRUCT()
+struct FAbstractSkeleton_AttributeBinding
+{
+	GENERATED_BODY()
+
+	UPROPERTY(EditAnywhere, Category = Default)
+	FName SetName;
+
+	UPROPERTY(EditAnywhere, Category = Default)
+	FAnimationAttributeIdentifier Attribute;
+};
+
+UCLASS(MinimalAPI, BlueprintType)
+class UAbstractSkeletonSetBinding : public UObject
+{
+	GENERATED_BODY()
+
+public:
+	SKELETONTEMPLATEFRAMEWORKRUNTIME_API bool SetSkeleton(const TObjectPtr<USkeleton> InSkeleton);
+
+	SKELETONTEMPLATEFRAMEWORKRUNTIME_API TObjectPtr<USkeleton> GetSkeleton() const;
+
+	SKELETONTEMPLATEFRAMEWORKRUNTIME_API bool SetSetCollection(const TObjectPtr<const UAbstractSkeletonSetCollection> InSetCollection);
+	
+	SKELETONTEMPLATEFRAMEWORKRUNTIME_API TObjectPtr<const UAbstractSkeletonSetCollection> GetSetCollection() const;
+	
+	// Bone Bindings
+
+	SKELETONTEMPLATEFRAMEWORKRUNTIME_API bool AddBoneToSet(const FName BoneName, const FName SetName);
+
+	SKELETONTEMPLATEFRAMEWORKRUNTIME_API bool RemoveBoneFromSet(const FName BoneName);
+
+	SKELETONTEMPLATEFRAMEWORKRUNTIME_API bool IsBoneInSet(const FName BoneName) const;
+
+	SKELETONTEMPLATEFRAMEWORKRUNTIME_API bool IsBoneInSet(const FName BoneName, const FName SetName) const;
+
+	// Returns the name of the set that the provided bone is in or NAME_None if bone is unbound.
+	SKELETONTEMPLATEFRAMEWORKRUNTIME_API FName GetBoneSet(const FName BoneName);
+
+	SKELETONTEMPLATEFRAMEWORKRUNTIME_API const TConstArrayView<FAbstractSkeleton_BoneBinding> GetBoneBindings() const;
+		
+	// Attribute Bindings
+
+	SKELETONTEMPLATEFRAMEWORKRUNTIME_API bool AddAttributeToSet(const FAnimationAttributeIdentifier Attribute, const FName SetName);
+
+	SKELETONTEMPLATEFRAMEWORKRUNTIME_API bool RemoveAttributeFromSet(const FAnimationAttributeIdentifier Attribute);
+	
+	// Returns true if the provided attribute is in ANY set
+	SKELETONTEMPLATEFRAMEWORKRUNTIME_API bool IsAttributeInSet(const FAnimationAttributeIdentifier Attribute) const;
+
+	// Returns true if the provided attribute is in the provided set
+	SKELETONTEMPLATEFRAMEWORKRUNTIME_API bool IsAttributeInSet(const FAnimationAttributeIdentifier Attribute, const FName SetName) const;
+
+	SKELETONTEMPLATEFRAMEWORKRUNTIME_API const TConstArrayView<FAbstractSkeleton_AttributeBinding> GetAttributeBindings() const;
+
+private:
+	UPROPERTY()
+	TObjectPtr<const UAbstractSkeletonSetCollection> SetCollection;
+	
+	UPROPERTY(AssetRegistrySearchable)
+	TObjectPtr<USkeleton> Skeleton;
+
+	UPROPERTY()
+	TArray<FAbstractSkeleton_BoneBinding> BoneBindings;
+
+	UPROPERTY()
+	TArray<FAbstractSkeleton_AttributeBinding> AttributeBindings;
+};
diff --git a/Engine/Plugins/Experimental/Animation/SkeletonTemplateFramework/Source/Runtime/Public/SkeletonTemplateFramework/AbstractSkeletonSetCollection.h b/Engine/Plugins/Experimental/Animation/SkeletonTemplateFramework/Source/Runtime/Public/SkeletonTemplateFramework/AbstractSkeletonSetCollection.h
new file mode 100644
--- /dev/null
+++ b/Engine/Plugins/Experimental/Animation/SkeletonTemplateFramework/Source/Runtime/Public/SkeletonTemplateFramework/AbstractSkeletonSetCollection.h
@@ -0,0 +1,54 @@
+// Copyright Epic Games, Inc. All Rights Reserved.
+
+#pragma once
+
+#include "CoreMinimal.h"
+
+#include "AbstractSkeletonSetCollection.generated.h"
+
+USTRUCT()
+struct FAbstractSkeletonSet
+{
+	GENERATED_BODY()
+
+	UPROPERTY(EditAnywhere, Category = Default)
+	FName SetName;
+
+	UPROPERTY(EditAnywhere, Category = Default)
+	FName ParentSetName;
+};
+
+UCLASS(MinimalAPI, BlueprintType)
+class UAbstractSkeletonSetCollection : public UObject
+{
+	GENERATED_BODY()
+
+public:
+	SKELETONTEMPLATEFRAMEWORKRUNTIME_API bool AddSet(const FName SetName, const FName ParentSetName = NAME_None);
+
+	SKELETONTEMPLATEFRAMEWORKRUNTIME_API bool HasSet(const FName SetName) const;
+
+	SKELETONTEMPLATEFRAMEWORKRUNTIME_API bool RenameSet(const FName OldSetName, const FName NewSetName);
+
+	SKELETONTEMPLATEFRAMEWORKRUNTIME_API bool RemoveSet(const FName SetName);
+
+	SKELETONTEMPLATEFRAMEWORKRUNTIME_API bool ReparentSet(const FName SetName, const FName NewParentName);
+
+	SKELETONTEMPLATEFRAMEWORKRUNTIME_API const TArrayView<const FAbstractSkeletonSet> GetSetHierarchy() const;
+
+	FName GetParentSetName(const FName SetName) const;
+
+private:
+	// Returns true if SetName is a descendant of AncestorSetName.
+	// This is for strict descendants, i.e. a set is not a descendant of itself.
+	bool IsDescendantOf(const FName SetName, const FName AncestorSetName) const;
+
+	int32 GetSetIndex(const FName SetName) const;
+
+	void SortSetHierarchy();
+
+	FAbstractSkeletonSet& GetSetRef(const FName SetName);
+
+	UPROPERTY()
+	TArray<FAbstractSkeletonSet> SetHierarchy; 
+};
