Fix missing copyright notice

#rnx

diff --git a/Engine/Plugins/Experimental/Animation/SkeletonTemplateFramework/Source/Runtime/Private/SkeletonTemplateFramework/Utils.h b/Engine/Plugins/Experimental/Animation/SkeletonTemplateFramework/Source/Runtime/Private/SkeletonTemplateFramework/Utils.h
--- a/Engine/Plugins/Experimental/Animation/SkeletonTemplateFramework/Source/Runtime/Private/SkeletonTemplateFramework/Utils.h
+++ b/Engine/Plugins/Experimental/Animation/SkeletonTemplateFramework/Source/Runtime/Private/SkeletonTemplateFramework/Utils.h
@@ -1,3 +1,5 @@
+// Copyright Epic Games, Inc. All Rights Reserved.
+
 #pragma once
 
 #include "CoreMinimal.h"
