Abstract Skeleton Editors

diff --git a/Engine/Plugins/Experimental/Animation/SkeletonTemplateFramework/Source/Editor/Private/SkeletonTemplateFramework/AbstractSkeletonAnimationEditorsExtender.cpp b/Engine/Plugins/Experimental/Animation/SkeletonTemplateFramework/Source/Editor/Private/SkeletonTemplateFramework/AbstractSkeletonAnimationEditorsExtender.cpp
new file mode 100644
--- /dev/null
+++ b/Engine/Plugins/Experimental/Animation/SkeletonTemplateFramework/Source/Editor/Private/SkeletonTemplateFramework/AbstractSkeletonAnimationEditorsExtender.cpp
@@ -0,0 +1,97 @@
+// Copyright Epic Games, Inc. All Rights Reserved.
+
+#include "AbstractSkeletonAnimationEditorsExtender.h"
+
+#include "AssetRegistry/AssetRegistryModule.h"
+#include "Modules/ModuleManager.h"
+#include "SkeletonTemplateFramework/AbstractSkeletonSetBinding.h"
+#include "SkeletonTemplateFramework/AbstractSkeletonLabelBinding.h"
+#include "Animation/Skeleton.h"
+#include "Engine/SkeletalMesh.h"
+#include "Styling/AppStyle.h"
+#include "AssetRegistry/AssetData.h"
+
+#define LOCTEXT_NAMESPACE "UE::Anim::SkelTemplates::UAnimationEditorsAssetFamilyExtension_AbstractSkeleton"
+
+TObjectPtr<UClass> UAnimationEditorsAssetFamilyExtension_AbstractSkeleton::GetAssetClass() const
+{
+	return UAbstractSkeletonSetBinding::StaticClass();
+}
+
+FText UAnimationEditorsAssetFamilyExtension_AbstractSkeleton::GetAssetTypeDisplayName() const
+{
+	return LOCTEXT("DisplayName", "Abstract Skeleton");
+}
+
+const FSlateBrush* UAnimationEditorsAssetFamilyExtension_AbstractSkeleton::GetAssetTypeDisplayIcon() const
+{
+	return FAppStyle::Get().GetBrush("Persona.AssetClass.SkeletalMesh");
+}
+
+void UAnimationEditorsAssetFamilyExtension_AbstractSkeleton::FindAssetsOfType(TArray<FAssetData>& OutAssets, const IAnimationEditorsAssetFamilyInterface& AssetFamilyInterface) const
+{
+	FAssetRegistryModule& AssetRegistryModule = FModuleManager::LoadModuleChecked<FAssetRegistryModule>("AssetRegistry");
+	FARFilter Filter;
+	Filter.bRecursiveClasses = true;
+	Filter.ClassPaths.Add(UAbstractSkeletonSetBinding::StaticClass()->GetClassPathName());
+	Filter.ClassPaths.Add(UAbstractSkeletonLabelBinding::StaticClass()->GetClassPathName());
+
+	if (TObjectPtr<const USkeleton> SkeletonAsset = AssetFamilyInterface.GetAssetOfType<USkeleton>())
+	{
+		const FAssetData SkeletonAssetData = FAssetData(SkeletonAsset.Get(), true);
+		if (SkeletonAssetData.IsValid())
+		{
+			Filter.TagsAndValues.Add(
+				TEXT("Skeleton"),
+				SkeletonAssetData.GetExportTextName());
+		}
+	}
+
+	AssetRegistryModule.Get().GetAssets(Filter, OutAssets);
+}
+
+bool UAnimationEditorsAssetFamilyExtension_AbstractSkeleton::IsAssetCompatible(const FAssetData& InAssetData, const IAnimationEditorsAssetFamilyInterface& AssetFamilyInterface) const
+{
+	FAssetDataTagMapSharedView::FFindTagResult Result = InAssetData.TagsAndValues.FindTag(TEXT("Skeleton"));
+	if (Result.IsSet())
+	{
+		if (TObjectPtr<const USkeleton> SkeletonAsset = AssetFamilyInterface.GetAssetOfType<USkeleton>())
+		{
+			return Result.GetValue() == FSoftObjectPath(SkeletonAsset).ToString();
+		}
+	}
+
+	return false;
+}
+
+void UAnimationEditorsAssetFamilyExtension_AbstractSkeleton::FindCounterpartAssets(const UObject* InAsset, IAnimationEditorsAssetFamilyInterface& AssetFamilyInterface)
+{
+	TObjectPtr<USkeleton> Skeleton;
+
+	if (const UAbstractSkeletonSetBinding* SetBinding = Cast<const UAbstractSkeletonSetBinding>(InAsset))
+	{
+		Skeleton = SetBinding->GetSkeleton();
+	}
+	else if (const UAbstractSkeletonLabelBinding* LabelBinding = Cast<const UAbstractSkeletonLabelBinding>(InAsset))
+	{
+		Skeleton = LabelBinding->GetSkeleton();
+	}
+
+	if (AssetFamilyInterface.IsAssetTypeInFamily<USkeleton>())
+	{
+		AssetFamilyInterface.SetAssetOfType<USkeleton>(Skeleton);
+	}
+
+	if (AssetFamilyInterface.IsAssetTypeInFamilyAndUnassigned<USkeletalMesh>() && Skeleton)
+	{
+		AssetFamilyInterface.SetAssetOfType<USkeletalMesh>(Skeleton->GetPreviewMesh());
+	}
+};
+
+void UAnimationEditorsAssetFamilyExtension_AbstractSkeleton::GetPosition(FName& OutBeforeClass, FName& OutAfterClass) const
+{
+	OutAfterClass = USkeleton::StaticClass()->GetFName();
+	OutBeforeClass = USkeletalMesh::StaticClass()->GetFName();
+}
+
+#undef LOCTEXT_NAMESPACE
diff --git a/Engine/Plugins/Experimental/Animation/SkeletonTemplateFramework/Source/Editor/Private/SkeletonTemplateFramework/AbstractSkeletonAnimationEditorsExtender.h b/Engine/Plugins/Experimental/Animation/SkeletonTemplateFramework/Source/Editor/Private/SkeletonTemplateFramework/AbstractSkeletonAnimationEditorsExtender.h
new file mode 100644
--- /dev/null
+++ b/Engine/Plugins/Experimental/Animation/SkeletonTemplateFramework/Source/Editor/Private/SkeletonTemplateFramework/AbstractSkeletonAnimationEditorsExtender.h
@@ -0,0 +1,23 @@
+// Copyright Epic Games, Inc. All Rights Reserved.
+
+#pragma once
+
+#include "AnimationEditorsAssetFamilyExtension.h"
+
+#include "AbstractSkeletonAnimationEditorsExtender.generated.h"
+
+UCLASS()
+class UAnimationEditorsAssetFamilyExtension_AbstractSkeleton : public UAnimationEditorsAssetFamilyExtension
+{
+	GENERATED_BODY()
+
+public:
+	// UAnimationEditorsAssetFamilyExtension
+	virtual TObjectPtr<UClass> GetAssetClass() const override;
+	virtual FText GetAssetTypeDisplayName() const override;
+	virtual const FSlateBrush* GetAssetTypeDisplayIcon() const override;
+	virtual void FindAssetsOfType(TArray<FAssetData>& OutAssets, const IAnimationEditorsAssetFamilyInterface& AssetFamilyInterface) const override;
+	virtual bool IsAssetCompatible(const FAssetData& InAssetData, const IAnimationEditorsAssetFamilyInterface& AssetFamilyInterface) const override;
+	virtual void FindCounterpartAssets(const UObject* InAsset, IAnimationEditorsAssetFamilyInterface& AssetFamilyInterface) override;
+	virtual void GetPosition(FName& OutBeforeClass, FName& OutAfterClass) const override;
+};
\ No newline at end of file
diff --git a/Engine/Plugins/Experimental/Animation/SkeletonTemplateFramework/Source/Editor/Private/SkeletonTemplateFramework/AbstractSkeletonEditor.cpp b/Engine/Plugins/Experimental/Animation/SkeletonTemplateFramework/Source/Editor/Private/SkeletonTemplateFramework/AbstractSkeletonEditor.cpp
new file mode 100644
--- /dev/null
+++ b/Engine/Plugins/Experimental/Animation/SkeletonTemplateFramework/Source/Editor/Private/SkeletonTemplateFramework/AbstractSkeletonEditor.cpp
@@ -0,0 +1,294 @@
+// Copyright Epic Games, Inc. All Rights Reserved.
+
+#include "AbstractSkeletonEditor.h"
+
+#include "SkeletonTemplateFramework/AbstractSkeletonSetBinding.h"
+#include "AssetRegistry/AssetRegistryModule.h"
+#include "ToolMenus.h"
+#include "PersonaModule.h"
+#include "Modules/ModuleManager.h"
+#include "Widgets/Docking/SDockTab.h"
+#include "Modules/ModuleManager.h"
+#include "PropertyEditorModule.h"
+#include "ScopedTransaction.h"
+#include "Framework/Docking/TabManager.h"
+#include "Framework/MultiBox/MultiBoxBuilder.h"
+#include "IPersonaToolkit.h"
+#include "IPersonaEditorModeManager.h"
+#include "SkeletonTemplateFramework/AbstractSkeletonLabelBinding.h"
+#include "SkeletonTemplateFramework/Labels/SLabelsTab.h"
+#include "SkeletonTemplateFramework/Sets/SSetsTab.h"
+#include "Animation/Skeleton.h"
+
+#define LOCTEXT_NAMESPACE "UE::Anim::SkelTemplates::FAbstractSkeletonEditor"
+
+namespace UE::Anim::SkelTemplates
+{
+	namespace Tabs
+	{
+		static const FName ViewportId(TEXT("AbstractSkeletonBindingEditor_ViewportTab"));
+		static const FName DetailsId(TEXT("AbstractSkeletonBindingEditor_PropertiesTab"));
+		static const FName SetsId(TEXT("AbstractSkeletonBindingEditor_SetsTab"));
+		static const FName LabelsId(TEXT("AbstractSkeletonBindingEditor_LabelsTab"));
+	}
+
+	void FAbstractSkeletonEditor::InitEditor(const TArray<UObject*>& InObjects, const TSharedPtr<IToolkitHost> InToolkitHost)
+	{
+		TObjectPtr<UObject> InObject = InObjects[0];
+
+		bool bDefaultShowSetsTab = true;
+		FPersonaModule& PersonaModule = FModuleManager::LoadModuleChecked<FPersonaModule>("Persona");
+
+		if (InObject->IsA<UAbstractSkeletonSetBinding>())
+		{
+			SetBinding = CastChecked<UAbstractSkeletonSetBinding>(InObject);
+
+			PersonaToolkit = PersonaModule.CreatePersonaToolkit(
+				SetBinding.Get(),
+				FPersonaToolkitArgs(),
+				SetBinding->GetSkeleton());
+
+			// Try find a LabelBinding to open by default
+			if (SetBinding->GetSkeleton())
+			{
+				FAssetRegistryModule& AssetRegistryModule = FModuleManager::LoadModuleChecked<FAssetRegistryModule>("AssetRegistry");
+				IAssetRegistry& AssetRegistry = AssetRegistryModule.Get();
+
+				FARFilter Filter;
+				Filter.ClassPaths.Add(UAbstractSkeletonLabelBinding::StaticClass()->GetClassPathName());
+				Filter.TagsAndValues.Add(TEXT("Skeleton"), FAssetData(SetBinding->GetSkeleton()).GetExportTextName());
+
+				TArray<FAssetData> AssetList;
+				AssetRegistry.GetAssets(Filter, AssetList);
+
+				if (!AssetList.IsEmpty())
+				{
+					LabelBinding = CastChecked<UAbstractSkeletonLabelBinding>(AssetList[0].GetAsset());
+				}
+			}
+		}
+		else if (InObject->IsA<UAbstractSkeletonLabelBinding>())
+		{
+			LabelBinding = CastChecked<UAbstractSkeletonLabelBinding>(InObject);
+			bDefaultShowSetsTab = false;
+
+			PersonaToolkit = PersonaModule.CreatePersonaToolkit(
+				LabelBinding.Get(),
+				FPersonaToolkitArgs(),
+				LabelBinding->GetSkeleton().Get());
+
+			// Try find a SetBinding to open by default
+			if (LabelBinding->GetSkeleton())
+			{
+				FAssetRegistryModule& AssetRegistryModule = FModuleManager::LoadModuleChecked<FAssetRegistryModule>("AssetRegistry");
+				IAssetRegistry& AssetRegistry = AssetRegistryModule.Get();
+
+				FARFilter Filter;
+				Filter.ClassPaths.Add(UAbstractSkeletonSetBinding::StaticClass()->GetClassPathName());
+				Filter.TagsAndValues.Add(TEXT("Skeleton"), FAssetData(LabelBinding->GetSkeleton()).GetExportTextName());
+
+				TArray<FAssetData> AssetList;
+				AssetRegistry.GetAssets(Filter, AssetList);
+
+				if (!AssetList.IsEmpty())
+				{
+					SetBinding = CastChecked<UAbstractSkeletonSetBinding>(AssetList[0].GetAsset());
+				}
+			}
+		}
+
+		const TSharedRef<FTabManager::FLayout> Layout = FTabManager::NewLayout("AbstractSkeletonBindingEditorToolkit")
+			->AddArea
+			(
+				FTabManager::NewPrimaryArea()->SetOrientation(Orient_Horizontal)
+				->Split
+				(
+					FTabManager::NewStack()
+					->SetSizeCoefficient(0.3f)
+					->SetHideTabWell(true)
+					->AddTab(Tabs::ViewportId, ETabState::OpenedTab)
+				)
+				->Split
+				(
+					FTabManager::NewStack()
+					->SetSizeCoefficient(0.4f)
+					->AddTab(Tabs::SetsId, ETabState::OpenedTab)
+					->AddTab(Tabs::LabelsId, ETabState::OpenedTab)
+					->SetForegroundTab(bDefaultShowSetsTab ? Tabs::SetsId : Tabs::LabelsId)
+				)
+				->Split
+				(
+					FTabManager::NewStack()
+					->SetSizeCoefficient(0.3f)
+					->AddTab(Tabs::DetailsId, ETabState::OpenedTab)
+				)
+			);
+
+
+		FAssetEditorToolkit::InitAssetEditor(EToolkitMode::Standalone, InToolkitHost, "AbstractSkeletonEditor", Layout, true, true, InObjects);
+
+		ExtendToolbar();
+		RegenerateMenusAndToolbars();
+
+		if (InObject->IsA<UAbstractSkeletonSetBinding>())
+		{
+			if (SetsTab.IsValid())
+			{
+				SetsTab->ActivateInParent(ETabActivationCause::SetDirectly);
+			}
+			else
+			{
+				GetAssociatedTabManager()->TryInvokeTab(Tabs::SetsId);
+			}		
+		}
+		else
+		{
+			if (LabelsTab.IsValid())
+			{
+				LabelsTab->ActivateInParent(ETabActivationCause::SetDirectly);
+			}
+			else
+			{
+				GetAssociatedTabManager()->TryInvokeTab(Tabs::LabelsId);
+			}
+		}
+	}
+
+	void FAbstractSkeletonEditor::RegisterTabSpawners(const TSharedRef<class FTabManager>& InTabManager)
+	{
+		FAssetEditorToolkit::RegisterTabSpawners(InTabManager);
+
+		CreateEditorModeManager();
+
+		WorkspaceMenuCategory = InTabManager->AddLocalWorkspaceMenuCategory(LOCTEXT("AbstractSkeletonEditorToolkit", "Abstract Skeleton Editor"));
+
+		// Sets tab
+		InTabManager->RegisterTabSpawner(Tabs::SetsId, FOnSpawnTab::CreateLambda([this](const FSpawnTabArgs&)
+			{
+				TSharedRef<SDockTab> Tab = SNew(SDockTab);
+				SetsTab = Tab;
+
+				TSharedRef<SWidget> TabWidget = SNew(Sets::SSetsTab)
+					.ParentTab(Tab)
+					.SetBinding(SetBinding);
+
+				Tab->SetContent(TabWidget);
+
+				return Tab;
+			}))
+			.SetDisplayName(LOCTEXT("SetsTab_DisplayName", "Sets"))
+			.SetGroup(WorkspaceMenuCategory.ToSharedRef())
+			.SetIcon(FSlateIcon(FAppStyle::GetAppStyleSetName(), ("ClassIcon.GroupActor")));
+
+		// Labels tab
+		InTabManager->RegisterTabSpawner(Tabs::LabelsId, FOnSpawnTab::CreateLambda([this](const FSpawnTabArgs&)
+			{
+				TSharedRef<SDockTab> Tab = SNew(SDockTab);
+				LabelsTab = Tab;
+
+				TSharedRef<SWidget> TabWidget = SNew(Labels::SLabelsTab)
+					.ParentTab(Tab)
+					.LabelBinding(LabelBinding)
+					.AbstractSkeletonEditor(SharedThis(this));
+
+				Tab->SetContent(TabWidget);
+
+				return Tab;
+			}))
+			.SetDisplayName(LOCTEXT("LabelsTab_DisplayName", "Labels"))
+			.SetGroup(WorkspaceMenuCategory.ToSharedRef())
+			.SetIcon(FSlateIcon(FAppStyle::GetAppStyleSetName(), ("Icons.Tag")));
+
+		// Viewport tab
+		InTabManager->RegisterTabSpawner(Tabs::ViewportId, FOnSpawnTab::CreateLambda([this](const FSpawnTabArgs&)
+			{
+				FPersonaModule& PersonaModule = FModuleManager::LoadModuleChecked<FPersonaModule>("Persona");
+
+				FPersonaViewportArgs PersonaViewportArgs(PersonaToolkit->GetPreviewScene());
+
+				TSharedRef<FWorkflowTabFactory> ViewportTabFactory = PersonaModule.CreatePersonaViewportTabFactory(SharedThis(this), PersonaViewportArgs);
+
+				FWorkflowTabSpawnInfo SpawnInfo;
+
+				return SNew(SDockTab)
+				[
+					ViewportTabFactory->CreateTabBody(SpawnInfo)
+				];
+			}))
+			.SetDisplayName(LOCTEXT("ViewportTab_DisplayName", "Viewport"))
+			.SetGroup(WorkspaceMenuCategory.ToSharedRef());
+	}
+
+	void FAbstractSkeletonEditor::UnregisterTabSpawners(const TSharedRef<class FTabManager>& InTabManager)
+	{
+		FAssetEditorToolkit::UnregisterTabSpawners(InTabManager);
+		InTabManager->UnregisterTabSpawner(Tabs::SetsId);
+		InTabManager->UnregisterTabSpawner(Tabs::LabelsId);
+		InTabManager->UnregisterTabSpawner(Tabs::ViewportId);
+		InTabManager->UnregisterTabSpawner(Tabs::DetailsId);
+	}
+	
+	void FAbstractSkeletonEditor::ExtendToolbar()
+	{
+		TSharedPtr<FExtender> ToolbarExtender = MakeShareable(new FExtender);
+		AddToolbarExtender(ToolbarExtender);
+
+		ToolbarExtender->AddToolBarExtension(
+			"Asset",
+			EExtensionHook::After,
+			GetToolkitCommands(),
+			FToolBarExtensionDelegate::CreateLambda([this](FToolBarBuilder& ToolbarBuilder)
+				{
+					FPersonaModule& PersonaModule = FModuleManager::LoadModuleChecked<FPersonaModule>("Persona");
+
+					if (SetBinding.IsValid())
+					{
+						TSharedRef<class IAssetFamily> AssetFamily = PersonaModule.CreatePersonaAssetFamily(SetBinding->GetSkeleton().Get());
+						TSharedRef<SWidget> Widget = PersonaModule.CreateAssetFamilyShortcutWidget(SharedThis(this), AssetFamily);
+						AddToolbarWidget(Widget);
+					}
+					else if (LabelBinding.IsValid())
+					{
+						TSharedRef<class IAssetFamily> AssetFamily = PersonaModule.CreatePersonaAssetFamily(LabelBinding->GetSkeleton().Get());
+						TSharedRef<SWidget> Widget = PersonaModule.CreateAssetFamilyShortcutWidget(SharedThis(this), AssetFamily);
+						AddToolbarWidget(Widget);
+					}
+				}
+			));
+	}
+
+	void FAbstractSkeletonEditor::CreateEditorModeManager()
+	{
+		FPersonaModule& PersonaModule = FModuleManager::LoadModuleChecked<FPersonaModule>("Persona");
+		TSharedPtr<IPersonaEditorModeManager> PersonaModeManager = MakeShareable(PersonaModule.CreatePersonaEditorModeManager());
+
+		EditorModeManager = PersonaModeManager;
+	}
+
+	FText FAbstractSkeletonEditor::GetToolkitName() const
+	{
+		return LOCTEXT("ToolkitName", "Abstract Skeleton Editor");
+	}
+
+	const FSlateBrush* FAbstractSkeletonEditor::GetDefaultTabIcon() const
+	{
+		return FAppStyle::Get().GetBrush("SkeletonTree.Bone");
+	}
+
+	FLinearColor FAbstractSkeletonEditor::GetDefaultTabColor() const
+	{
+		return FLinearColor::White;
+	}
+
+	FText FAbstractSkeletonEditor::GetToolkitToolTipText() const
+	{
+		return LOCTEXT("ToolkitName", "Editor for manipulating a Skeleton's Set Bindings and Label Bindings");
+	}
+
+	TSharedPtr<IPersonaToolkit> FAbstractSkeletonEditor::GetPersonaToolkit() const
+	{
+		return PersonaToolkit;
+	}
+}
+
+#undef LOCTEXT_NAMESPACE
\ No newline at end of file
diff --git a/Engine/Plugins/Experimental/Animation/SkeletonTemplateFramework/Source/Editor/Private/SkeletonTemplateFramework/AbstractSkeletonEditor.h b/Engine/Plugins/Experimental/Animation/SkeletonTemplateFramework/Source/Editor/Private/SkeletonTemplateFramework/AbstractSkeletonEditor.h
new file mode 100644
--- /dev/null
+++ b/Engine/Plugins/Experimental/Animation/SkeletonTemplateFramework/Source/Editor/Private/SkeletonTemplateFramework/AbstractSkeletonEditor.h
@@ -0,0 +1,53 @@
+// Copyright Epic Games, Inc. All Rights Reserved.
+
+#pragma once
+
+#include "CoreMinimal.h"
+#include "SkeletonTemplateFramework/IAbstractSkeletonEditor.h"
+#include "WorkflowOrientedApp/WorkflowCentricApplication.h"
+
+class IPersonaToolkit;
+class UAbstractSkeletonLabelBinding;
+class UAbstractSkeletonSetBinding;
+
+namespace UE::Anim::SkelTemplates
+{
+	class FAbstractSkeletonEditor : public FWorkflowCentricApplication, public IAbstractSkeletonEditor
+	{
+	public:
+		void InitEditor(const TArray<UObject*>& InObjects, const TSharedPtr<IToolkitHost> InToolkitHost);
+
+		// Begin FWorkflowCentricApplication
+		virtual void RegisterTabSpawners(const TSharedRef<class FTabManager>& TabManager) override;
+		virtual void UnregisterTabSpawners(const TSharedRef<class FTabManager>& TabManager) override;
+		virtual void CreateEditorModeManager() override;
+
+		virtual FText GetToolkitName() const override;
+		virtual const FSlateBrush* GetDefaultTabIcon() const override;
+		virtual FLinearColor GetDefaultTabColor() const override;
+		virtual FText GetToolkitToolTipText() const override;
+		virtual FName GetToolkitFName() const override { return "AbstractSkeletonEditor"; }
+		virtual FText GetBaseToolkitName() const override { return INVTEXT("Abstract Skeleton Editor"); }
+		virtual FString GetWorldCentricTabPrefix() const override { return "Abstract Skeleton Editor "; }
+		virtual FLinearColor GetWorldCentricTabColorScale() const override { return {}; }
+		// End FWorkflowCentricApplication
+
+		// Begin IAbstractSkeletonEditor
+		virtual TSharedPtr<IPersonaToolkit> GetPersonaToolkit() const override;
+		// End IAbstractSkeletonEditor
+
+	private:
+		void ExtendToolbar();
+
+		TWeakObjectPtr<UAbstractSkeletonSetBinding> SetBinding;
+
+		TWeakObjectPtr<UAbstractSkeletonLabelBinding> LabelBinding;
+
+		TSharedPtr<IPersonaToolkit> PersonaToolkit;
+
+		TSharedPtr<SDockTab> SetsTab;
+
+		TSharedPtr<SDockTab> LabelsTab;
+	};
+
+}
diff --git a/Engine/Plugins/Experimental/Animation/SkeletonTemplateFramework/Source/Editor/Private/SkeletonTemplateFramework/IAbstractSkeletonEditor.h b/Engine/Plugins/Experimental/Animation/SkeletonTemplateFramework/Source/Editor/Private/SkeletonTemplateFramework/IAbstractSkeletonEditor.h
new file mode 100644
--- /dev/null
+++ b/Engine/Plugins/Experimental/Animation/SkeletonTemplateFramework/Source/Editor/Private/SkeletonTemplateFramework/IAbstractSkeletonEditor.h
@@ -0,0 +1,18 @@
+// Copyright Epic Games, Inc. All Rights Reserved.
+
+#pragma once
+
+#include "CoreMinimal.h"
+
+class IPersonaToolkit;
+
+namespace UE::Anim::SkelTemplates
+{
+
+	class IAbstractSkeletonEditor
+	{
+	public:
+		virtual TSharedPtr<IPersonaToolkit> GetPersonaToolkit() const = 0;
+	};
+
+}
\ No newline at end of file
diff --git a/Engine/Plugins/Experimental/Animation/SkeletonTemplateFramework/Source/Editor/Private/SkeletonTemplateFramework/Labels/BoneDragDrop.h b/Engine/Plugins/Experimental/Animation/SkeletonTemplateFramework/Source/Editor/Private/SkeletonTemplateFramework/Labels/BoneDragDrop.h
new file mode 100644
--- /dev/null
+++ b/Engine/Plugins/Experimental/Animation/SkeletonTemplateFramework/Source/Editor/Private/SkeletonTemplateFramework/Labels/BoneDragDrop.h
@@ -0,0 +1,103 @@
+// Copyright Epic Games, Inc. All Rights Reserved.
+
+#pragma once
+
+#include "CoreMinimal.h"
+#include "DragAndDrop/DecoratedDragDropOp.h"
+#include "Widgets/Text/SRichTextBlock.h"
+
+class UAbstractSkeletonLabelBinding;
+
+#define LOCTEXT_NAMESPACE "UE::Anim::SkelTemplatesLabels::FBoneDragDropOp"
+
+namespace UE::Anim::SkelTemplates::Labels
+{
+	class FBoneDragDropOp : public FDecoratedDragDropOp
+	{
+	public:
+		DRAG_DROP_OPERATOR_TYPE(FBoneDragDropOp, FDecoratedDragDropOp)
+
+		static TSharedRef<FBoneDragDropOp> New(const TWeakObjectPtr<UAbstractSkeletonLabelBinding> InLabelBinding, const FName InBoneName)
+		{
+			TSharedRef<FBoneDragDropOp> Operation = MakeShared<FBoneDragDropOp>();
+			Operation->LabelBinding = InLabelBinding;
+			Operation->BoneName = InBoneName;
+
+			Operation->Construct();
+
+			return Operation;
+		}
+
+		TSharedPtr<SWidget> GetDefaultDecorator() const override
+		{
+			return SNew(SBorder)
+				.Visibility(EVisibility::Visible)
+				.BorderImage(FAppStyle::GetBrush("Menu.Background"))
+				.Padding(FMargin(8.0f, 4.0f))
+				[
+					SNew(SHorizontalBox)
+					+ SHorizontalBox::Slot()
+					.AutoWidth()
+					.Padding(FMargin(2.0f))
+					[
+						SNew(SImage)
+							.Visibility_Lambda([this]()
+								{
+									return HoveredLabel != NAME_None ? EVisibility::Collapsed : EVisibility::Visible;
+								})
+							.Image(FAppStyle::Get().GetBrush("SkeletonTree.Bone"))
+					]
+					+ SHorizontalBox::Slot()
+					.AutoWidth()
+					.Padding(FMargin(2.0f))
+					[
+						SNew(STextBlock)
+							.Visibility_Lambda([this]()
+								{
+									return HoveredLabel != NAME_None ? EVisibility::Collapsed : EVisibility::Visible;
+								})
+							.Text(FText::FromName(BoneName))
+					]
+					+ SHorizontalBox::Slot()
+					.AutoWidth()
+					.Padding(FMargin(2.0f))
+					[
+						SNew(SImage)
+							.Visibility_Lambda([this]()
+								{
+									return HoveredLabel == NAME_None ? EVisibility::Collapsed : EVisibility::Visible;
+								})
+							.Image(FAppStyle::Get().GetBrush("Icons.Link"))
+							.ColorAndOpacity(FAppStyle::Get().GetSlateColor("Colors.AccentBlue"))
+					]
+					+ SHorizontalBox::Slot()
+					.AutoWidth()
+					.Padding(FMargin(2.0f))
+					[
+						SNew(SRichTextBlock)
+							.DecoratorStyleSet(&FAppStyle::Get())
+							.Visibility_Lambda([this]()
+								{
+									return HoveredLabel == NAME_None ? EVisibility::Collapsed : EVisibility::Visible;
+								})
+							.Text_Lambda([this]()
+								{
+									return FText::Format(LOCTEXT("BindBoneToLabelTooltip", "Bind <RichTextBlock.Bold>{0}</> to <RichTextBlock.Bold>{1}</>"), FText::FromName(BoneName), FText::FromName(HoveredLabel));
+								})
+					]
+					
+				];
+		}
+
+		void SetHoveredLabel(const FName InLabel)
+		{
+			HoveredLabel = InLabel;
+		}
+
+		TWeakObjectPtr<UAbstractSkeletonLabelBinding> LabelBinding;
+		FName BoneName;
+		FName HoveredLabel;
+	};
+}
+
+#undef LOCTEXT_NAMESPACE
\ No newline at end of file
diff --git a/Engine/Plugins/Experimental/Animation/SkeletonTemplateFramework/Source/Editor/Private/SkeletonTemplateFramework/Labels/ILabelBinding.h b/Engine/Plugins/Experimental/Animation/SkeletonTemplateFramework/Source/Editor/Private/SkeletonTemplateFramework/Labels/ILabelBinding.h
new file mode 100644
--- /dev/null
+++ b/Engine/Plugins/Experimental/Animation/SkeletonTemplateFramework/Source/Editor/Private/SkeletonTemplateFramework/Labels/ILabelBinding.h
@@ -0,0 +1,23 @@
+// Copyright Epic Games, Inc. All Rights Reserved.
+
+#pragma once
+
+#include "CoreMinimal.h"
+
+class UAbstractSkeletonLabelCollection;
+class UAbstractSkeletonLabelBinding;
+
+namespace UE::Anim::SkelTemplates::Labels
+{
+
+	class ILabelBindingWidget
+	{
+	public:
+		virtual void ScrollToLabel(const TObjectPtr<const UAbstractSkeletonLabelCollection> InLabelCollection, const FName InLabel) = 0;
+
+		virtual void RepopulateTreeData() = 0;
+
+		virtual TWeakObjectPtr<UAbstractSkeletonLabelBinding> GetLabelBinding() const = 0;
+	};
+
+}
\ No newline at end of file
diff --git a/Engine/Plugins/Experimental/Animation/SkeletonTemplateFramework/Source/Editor/Private/SkeletonTemplateFramework/Labels/ILabelSkeletonTree.h b/Engine/Plugins/Experimental/Animation/SkeletonTemplateFramework/Source/Editor/Private/SkeletonTemplateFramework/Labels/ILabelSkeletonTree.h
new file mode 100644
--- /dev/null
+++ b/Engine/Plugins/Experimental/Animation/SkeletonTemplateFramework/Source/Editor/Private/SkeletonTemplateFramework/Labels/ILabelSkeletonTree.h
@@ -0,0 +1,20 @@
+// Copyright Epic Games, Inc. All Rights Reserved.
+
+#pragma once
+
+#include "UObject/NameTypes.h"
+
+namespace UE::Anim::SkelTemplates::Labels
+{
+
+	class ILabelSkeletonTreeWidget
+	{
+	public:
+		virtual ~ILabelSkeletonTreeWidget() = default;
+		
+		virtual void ScrollToBone(const FName InBoneName) = 0;
+
+		virtual void RepopulateTreeData() = 0;
+	};
+
+}
\ No newline at end of file
diff --git a/Engine/Plugins/Experimental/Animation/SkeletonTemplateFramework/Source/Editor/Private/SkeletonTemplateFramework/Labels/ILabelsTab.h b/Engine/Plugins/Experimental/Animation/SkeletonTemplateFramework/Source/Editor/Private/SkeletonTemplateFramework/Labels/ILabelsTab.h
new file mode 100644
--- /dev/null
+++ b/Engine/Plugins/Experimental/Animation/SkeletonTemplateFramework/Source/Editor/Private/SkeletonTemplateFramework/Labels/ILabelsTab.h
@@ -0,0 +1,31 @@
+// Copyright Epic Games, Inc. All Rights Reserved.
+
+#pragma once
+
+#include "Templates/SharedPointer.h"
+
+namespace UE::Anim::SkelTemplates
+{
+	class IAbstractSkeletonEditor;
+}
+
+namespace UE::Anim::SkelTemplates::Labels
+{
+	class ILabelBindingWidget;
+	class ILabelSkeletonTreeWidget;
+
+	class ILabelsTab
+	{
+	public:
+		virtual ~ILabelsTab() = default;
+		
+		virtual TSharedPtr<ILabelBindingWidget> GetLabelBindingWidget() const = 0;
+
+		virtual TSharedPtr<ILabelSkeletonTreeWidget> GetLabelSkeletonTreeWidget() const = 0;
+
+		virtual void RepopulateLabelData() = 0;
+
+		virtual TWeakPtr<IAbstractSkeletonEditor> GetAbstractSkeletonEditor() const = 0;
+	};
+
+}
\ No newline at end of file
diff --git a/Engine/Plugins/Experimental/Animation/SkeletonTemplateFramework/Source/Editor/Private/SkeletonTemplateFramework/Labels/LabelBindingAssetDefinition.cpp b/Engine/Plugins/Experimental/Animation/SkeletonTemplateFramework/Source/Editor/Private/SkeletonTemplateFramework/Labels/LabelBindingAssetDefinition.cpp
new file mode 100644
--- /dev/null
+++ b/Engine/Plugins/Experimental/Animation/SkeletonTemplateFramework/Source/Editor/Private/SkeletonTemplateFramework/Labels/LabelBindingAssetDefinition.cpp
@@ -0,0 +1,43 @@
+// Copyright Epic Games, Inc. All Rights Reserved.
+
+#include "SkeletonTemplateFramework/Labels/LabelBindingAssetDefinition.h"
+
+#include "SkeletonTemplateFramework/AbstractSkeletonEditor.h"
+#include "SkeletonTemplateFramework/AbstractSkeletonLabelBinding.h"
+
+#define LOCTEXT_NAMESPACE "UE::Anim::SkelTemplates::UAssetDefinition_AbstractSkeletonLabelBinding"
+
+FText UAssetDefinition_AbstractSkeletonLabelBinding::GetAssetDisplayName() const
+{
+	return LOCTEXT("LabelBinding", "UAF Label Binding");
+}
+
+FLinearColor UAssetDefinition_AbstractSkeletonLabelBinding::GetAssetColor() const
+{
+	return FLinearColor(FColor(21, 119, 53));
+}
+
+TSoftClassPtr<UObject> UAssetDefinition_AbstractSkeletonLabelBinding::GetAssetClass() const
+{
+	return UAbstractSkeletonLabelBinding::StaticClass();
+}
+
+EAssetCommandResult UAssetDefinition_AbstractSkeletonLabelBinding::OpenAssets(const FAssetOpenArgs& OpenArgs) const
+{
+	if (OpenArgs.OpenMethod == EAssetOpenMethod::Edit)
+	{
+		TArray<UObject*> Assets = OpenArgs.LoadObjects<UObject>();
+
+		MakeShared<UE::Anim::SkelTemplates::FAbstractSkeletonEditor>()->InitEditor(Assets, OpenArgs.ToolkitHost);
+	}
+
+	return EAssetCommandResult::Handled;
+}
+
+TConstArrayView<FAssetCategoryPath> UAssetDefinition_AbstractSkeletonLabelBinding::GetAssetCategories() const
+{
+	static const auto Categories = { FAssetCategoryPath(EAssetCategoryPaths::Animation, LOCTEXT("UAFSubMenu", "Animation Framework")) };
+	return Categories;
+}
+
+#undef LOCTEXT_NAMESPACE
diff --git a/Engine/Plugins/Experimental/Animation/SkeletonTemplateFramework/Source/Editor/Private/SkeletonTemplateFramework/Labels/LabelBindingAssetDefinition.h b/Engine/Plugins/Experimental/Animation/SkeletonTemplateFramework/Source/Editor/Private/SkeletonTemplateFramework/Labels/LabelBindingAssetDefinition.h
new file mode 100644
--- /dev/null
+++ b/Engine/Plugins/Experimental/Animation/SkeletonTemplateFramework/Source/Editor/Private/SkeletonTemplateFramework/Labels/LabelBindingAssetDefinition.h
@@ -0,0 +1,22 @@
+// Copyright Epic Games, Inc. All Rights Reserved.
+
+#pragma once
+
+#include "AssetDefinitionDefault.h"
+
+#include "LabelBindingAssetDefinition.generated.h"
+
+UCLASS()
+class UAssetDefinition_AbstractSkeletonLabelBinding : public UAssetDefinitionDefault
+{
+	GENERATED_BODY()
+
+public:
+	// Begin UAssetDefinition
+	virtual FText GetAssetDisplayName() const override;
+	virtual FLinearColor GetAssetColor() const override;
+	virtual TSoftClassPtr<UObject> GetAssetClass() const override;
+	virtual EAssetCommandResult OpenAssets(const FAssetOpenArgs& OpenArgs) const override;
+	virtual TConstArrayView<FAssetCategoryPath> GetAssetCategories() const override;
+	// End UAssetDefinition
+};
diff --git a/Engine/Plugins/Experimental/Animation/SkeletonTemplateFramework/Source/Editor/Private/SkeletonTemplateFramework/Labels/LabelBindingFactory.cpp b/Engine/Plugins/Experimental/Animation/SkeletonTemplateFramework/Source/Editor/Private/SkeletonTemplateFramework/Labels/LabelBindingFactory.cpp
new file mode 100644
--- /dev/null
+++ b/Engine/Plugins/Experimental/Animation/SkeletonTemplateFramework/Source/Editor/Private/SkeletonTemplateFramework/Labels/LabelBindingFactory.cpp
@@ -0,0 +1,31 @@
+// Copyright Epic Games, Inc. All Rights Reserved.
+
+#include "SkeletonTemplateFramework/Labels/LabelBindingFactory.h"
+
+#include "SkeletonTemplateFramework/AbstractSkeletonLabelBinding.h"
+#include "UObject/Package.h"
+
+UAbstractSkeletonLabelBindingFactory::UAbstractSkeletonLabelBindingFactory()
+{
+	bCreateNew = true;
+	bEditAfterNew = true;
+	SupportedClass = UAbstractSkeletonLabelBinding::StaticClass();
+}
+
+bool UAbstractSkeletonLabelBindingFactory::ConfigureProperties()
+{
+	return true;
+}
+
+UObject* UAbstractSkeletonLabelBindingFactory::FactoryCreateNew(UClass* Class, UObject* InParent, FName Name, EObjectFlags Flags, UObject* Context, FFeedbackContext* Warn, FName CallingContext)
+{
+	EObjectFlags FlagsToUse = Flags | RF_Public | RF_Standalone | RF_Transactional | RF_LoadCompleted;
+	if (InParent == GetTransientPackage())
+	{
+		FlagsToUse &= ~RF_Standalone;
+	}
+
+	TObjectPtr<UAbstractSkeletonLabelBinding> NewLabelBinding = NewObject<UAbstractSkeletonLabelBinding>(InParent, Class, Name, FlagsToUse);
+
+	return NewLabelBinding;
+}
diff --git a/Engine/Plugins/Experimental/Animation/SkeletonTemplateFramework/Source/Editor/Private/SkeletonTemplateFramework/Labels/LabelBindingFactory.h b/Engine/Plugins/Experimental/Animation/SkeletonTemplateFramework/Source/Editor/Private/SkeletonTemplateFramework/Labels/LabelBindingFactory.h
new file mode 100644
--- /dev/null
+++ b/Engine/Plugins/Experimental/Animation/SkeletonTemplateFramework/Source/Editor/Private/SkeletonTemplateFramework/Labels/LabelBindingFactory.h
@@ -0,0 +1,23 @@
+// Copyright Epic Games, Inc. All Rights Reserved.
+
+#pragma once
+
+#include "CoreMinimal.h"
+#include "Factories/Factory.h"
+
+#include "LabelBindingFactory.generated.h"
+
+UCLASS()
+class UAbstractSkeletonLabelBindingFactory : public UFactory
+{
+	GENERATED_BODY()
+
+	UAbstractSkeletonLabelBindingFactory();
+
+	virtual bool ConfigureProperties() override;
+	virtual UObject* FactoryCreateNew(UClass* Class, UObject* InParent, FName Name, EObjectFlags Flags, UObject* Context, FFeedbackContext* Warn, FName CallingContext) override;
+	virtual UObject* FactoryCreateNew(UClass* Class, UObject* InParent, FName Name, EObjectFlags Flags, UObject* Context, FFeedbackContext* Warn) override
+	{
+		return FactoryCreateNew(Class, InParent, Name, Flags, Context, Warn, NAME_None);
+	}
+};
diff --git a/Engine/Plugins/Experimental/Animation/SkeletonTemplateFramework/Source/Editor/Private/SkeletonTemplateFramework/Labels/LabelCollectionAssetDefinition.cpp b/Engine/Plugins/Experimental/Animation/SkeletonTemplateFramework/Source/Editor/Private/SkeletonTemplateFramework/Labels/LabelCollectionAssetDefinition.cpp
new file mode 100644
--- /dev/null
+++ b/Engine/Plugins/Experimental/Animation/SkeletonTemplateFramework/Source/Editor/Private/SkeletonTemplateFramework/Labels/LabelCollectionAssetDefinition.cpp
@@ -0,0 +1,43 @@
+// Copyright Epic Games, Inc. All Rights Reserved.
+
+#include "SkeletonTemplateFramework/Labels/LabelCollectionAssetDefinition.h"
+
+#include "SkeletonTemplateFramework/AbstractSkeletonLabelCollection.h"
+#include "SkeletonTemplateFramework/Labels/LabelCollectionEditor.h"
+
+#define LOCTEXT_NAMESPACE "UE::Anim::SkelTemplates"
+
+FText UAssetDefinition_AbstractSkeletonLabelCollection::GetAssetDisplayName() const
+{
+	return LOCTEXT("LabelCollection", "UAF Label Collection");
+}
+
+FLinearColor UAssetDefinition_AbstractSkeletonLabelCollection::GetAssetColor() const
+{
+	return FLinearColor(FColor(73, 148, 98));
+}
+
+TSoftClassPtr<UObject> UAssetDefinition_AbstractSkeletonLabelCollection::GetAssetClass() const
+{
+	return UAbstractSkeletonLabelCollection::StaticClass();
+}
+
+EAssetCommandResult UAssetDefinition_AbstractSkeletonLabelCollection::OpenAssets(const FAssetOpenArgs& OpenArgs) const
+{
+	if (OpenArgs.OpenMethod == EAssetOpenMethod::Edit)
+	{
+		TArray<UObject*> Assets = OpenArgs.LoadObjects<UObject>();
+		
+		MakeShared<UE::Anim::SkelTemplates::Labels::FLabelCollectionEditorToolkit>()->InitEditor(Assets);
+	}
+
+	return EAssetCommandResult::Handled;
+}
+
+TConstArrayView<FAssetCategoryPath> UAssetDefinition_AbstractSkeletonLabelCollection::GetAssetCategories() const
+{
+	static const auto Categories = { FAssetCategoryPath(EAssetCategoryPaths::Animation, LOCTEXT("UAFSubMenu", "Animation Framework")) };
+	return Categories;
+}
+
+#undef LOCTEXT_NAMESPACE
diff --git a/Engine/Plugins/Experimental/Animation/SkeletonTemplateFramework/Source/Editor/Private/SkeletonTemplateFramework/Labels/LabelCollectionAssetDefinition.h b/Engine/Plugins/Experimental/Animation/SkeletonTemplateFramework/Source/Editor/Private/SkeletonTemplateFramework/Labels/LabelCollectionAssetDefinition.h
new file mode 100644
--- /dev/null
+++ b/Engine/Plugins/Experimental/Animation/SkeletonTemplateFramework/Source/Editor/Private/SkeletonTemplateFramework/Labels/LabelCollectionAssetDefinition.h
@@ -0,0 +1,22 @@
+// Copyright Epic Games, Inc. All Rights Reserved.
+
+#pragma once
+
+#include "AssetDefinitionDefault.h"
+
+#include "LabelCollectionAssetDefinition.generated.h"
+
+UCLASS()
+class UAssetDefinition_AbstractSkeletonLabelCollection : public UAssetDefinitionDefault
+{
+	GENERATED_BODY()
+
+public:
+	// Begin UAssetDefinition
+	virtual FText GetAssetDisplayName() const override;
+	virtual FLinearColor GetAssetColor() const override;
+	virtual TSoftClassPtr<UObject> GetAssetClass() const override;
+	virtual EAssetCommandResult OpenAssets(const FAssetOpenArgs& OpenArgs) const override;
+	virtual TConstArrayView<FAssetCategoryPath> GetAssetCategories() const override;
+	// End UAssetDefinition
+};
diff --git a/Engine/Plugins/Experimental/Animation/SkeletonTemplateFramework/Source/Editor/Private/SkeletonTemplateFramework/Labels/LabelCollectionEditor.cpp b/Engine/Plugins/Experimental/Animation/SkeletonTemplateFramework/Source/Editor/Private/SkeletonTemplateFramework/Labels/LabelCollectionEditor.cpp
new file mode 100644
--- /dev/null
+++ b/Engine/Plugins/Experimental/Animation/SkeletonTemplateFramework/Source/Editor/Private/SkeletonTemplateFramework/Labels/LabelCollectionEditor.cpp
@@ -0,0 +1,60 @@
+// Copyright Epic Games, Inc. All Rights Reserved.
+
+#include "SkeletonTemplateFramework/Labels/LabelCollectionEditor.h"
+
+#include "SkeletonTemplateFramework/AbstractSkeletonLabelCollection.h"
+#include "SkeletonTemplateFramework/Labels/SLabelCollection.h"
+#include "Widgets/Docking/SDockTab.h"
+
+#define LOCTEXT_NAMESPACE "UE::Anim::SkelTemplates::FLabelCollectionEditorToolkit"
+
+namespace UE::Anim::SkelTemplates::Labels
+{
+	static FName LabelsTabId("AbstractSkeletonLabelCollection_LabelsTab");
+
+	void FLabelCollectionEditorToolkit::InitEditor(const TArray<UObject*>& InObjects)
+	{
+		LabelCollection = CastChecked<UAbstractSkeletonLabelCollection>(InObjects[0]);
+
+		const TSharedRef<FTabManager::FLayout> Layout = FTabManager::NewLayout("AbstractSkeletonLabelCollectionEditorToolkit")
+			->AddArea
+			(
+				FTabManager::NewPrimaryArea()
+				->SetOrientation(Orient_Horizontal)
+				->Split
+				(
+					FTabManager::NewStack()
+					->AddTab(LabelsTabId, ETabState::OpenedTab)
+					->SetHideTabWell(true)
+				)
+			);
+
+		FAssetEditorToolkit::InitAssetEditor(EToolkitMode::Standalone, {}, "AbstractSkeletonLabelCollectionEditorToolkit", Layout, true, true, InObjects);
+	}
+
+	void FLabelCollectionEditorToolkit::RegisterTabSpawners(const TSharedRef<class FTabManager>& InTabManager)
+	{
+		FAssetEditorToolkit::RegisterTabSpawners(InTabManager);
+
+		WorkspaceMenuCategory = InTabManager->AddLocalWorkspaceMenuCategory(LOCTEXT("AbstractSkeletonLabelsCollectionEditor", "Abstract Skeleton Labels Editor"));
+
+		InTabManager->RegisterTabSpawner(LabelsTabId, FOnSpawnTab::CreateLambda([this](const FSpawnTabArgs&)
+			{
+				return SNew(SDockTab)
+				[
+					SNew(SLabelCollection, LabelCollection)
+				];
+			}))
+			.SetDisplayName(LOCTEXT("LabelsTab_DisplayName", "Labels"))
+			.SetIcon(FSlateIcon(FAppStyle::GetAppStyleSetName(), "Icons.Tag"))
+			.SetGroup(WorkspaceMenuCategory.ToSharedRef());
+	}
+
+	void FLabelCollectionEditorToolkit::UnregisterTabSpawners(const TSharedRef<class FTabManager>& InTabManager)
+	{
+		FAssetEditorToolkit::UnregisterTabSpawners(InTabManager);
+		InTabManager->UnregisterTabSpawner(LabelsTabId);
+	}
+}
+
+#undef LOCTEXT_NAMESPACE
\ No newline at end of file
diff --git a/Engine/Plugins/Experimental/Animation/SkeletonTemplateFramework/Source/Editor/Private/SkeletonTemplateFramework/Labels/LabelCollectionEditor.h b/Engine/Plugins/Experimental/Animation/SkeletonTemplateFramework/Source/Editor/Private/SkeletonTemplateFramework/Labels/LabelCollectionEditor.h
new file mode 100644
--- /dev/null
+++ b/Engine/Plugins/Experimental/Animation/SkeletonTemplateFramework/Source/Editor/Private/SkeletonTemplateFramework/Labels/LabelCollectionEditor.h
@@ -0,0 +1,32 @@
+// Copyright Epic Games, Inc. All Rights Reserved.
+
+#pragma once
+
+#include "CoreMinimal.h"
+#include "Toolkits/AssetEditorToolkit.h"
+
+class UAbstractSkeletonLabelCollection;
+
+namespace UE::Anim::SkelTemplates::Labels
+{
+	class FLabelCollectionEditorToolkit : public FAssetEditorToolkit
+	{
+	public:
+		void InitEditor(const TArray<UObject*>& InObjects);
+
+		// Begin FAssetEditorToolkit
+		void RegisterTabSpawners(const TSharedRef<class FTabManager>& TabManager) override;
+		void UnregisterTabSpawners(const TSharedRef<class FTabManager>& TabManager) override;
+
+		FName GetToolkitFName() const override { return "AbstractSkeletonLabelCollectionEditor"; }
+		FText GetBaseToolkitName() const override { return INVTEXT("Abstract Skeleton Label Collection Editor"); }
+		FString GetWorldCentricTabPrefix() const override { return "Abstract Skeleton Label Collection Editor "; }
+		FLinearColor GetWorldCentricTabColorScale() const override { return {}; }
+		// End FAssetEditorToolkit
+
+	private:
+		TSharedRef<SDockTab> SpawnPropertiesTab(const FSpawnTabArgs& Args);
+
+		TWeakObjectPtr<UAbstractSkeletonLabelCollection> LabelCollection;
+	};
+}
\ No newline at end of file
diff --git a/Engine/Plugins/Experimental/Animation/SkeletonTemplateFramework/Source/Editor/Private/SkeletonTemplateFramework/Labels/LabelCollectionFactory.cpp b/Engine/Plugins/Experimental/Animation/SkeletonTemplateFramework/Source/Editor/Private/SkeletonTemplateFramework/Labels/LabelCollectionFactory.cpp
new file mode 100644
--- /dev/null
+++ b/Engine/Plugins/Experimental/Animation/SkeletonTemplateFramework/Source/Editor/Private/SkeletonTemplateFramework/Labels/LabelCollectionFactory.cpp
@@ -0,0 +1,31 @@
+// Copyright Epic Games, Inc. All Rights Reserved.
+
+#include "SkeletonTemplateFramework/Labels/LabelCollectionFactory.h"
+
+#include "SkeletonTemplateFramework/AbstractSkeletonLabelCollection.h"
+#include "UObject/Package.h"
+
+UAbstractSkeletonLabelCollectionFactory::UAbstractSkeletonLabelCollectionFactory()
+{
+	bCreateNew = true;
+	bEditAfterNew = true;
+	SupportedClass = UAbstractSkeletonLabelCollection::StaticClass();
+}
+
+bool UAbstractSkeletonLabelCollectionFactory::ConfigureProperties()
+{
+	return true;
+}
+
+UObject* UAbstractSkeletonLabelCollectionFactory::FactoryCreateNew(UClass* Class, UObject* InParent, FName Name, EObjectFlags Flags, UObject* Context, FFeedbackContext* Warn, FName CallingContext)
+{
+	EObjectFlags FlagsToUse = Flags | RF_Public | RF_Standalone | RF_Transactional | RF_LoadCompleted;
+	if (InParent == GetTransientPackage())
+	{
+		FlagsToUse &= ~RF_Standalone;
+	}
+
+	TObjectPtr<UAbstractSkeletonLabelCollection> NewLabelCollection = NewObject<UAbstractSkeletonLabelCollection>(InParent, Class, Name, FlagsToUse);
+
+	return NewLabelCollection;
+}
diff --git a/Engine/Plugins/Experimental/Animation/SkeletonTemplateFramework/Source/Editor/Private/SkeletonTemplateFramework/Labels/LabelCollectionFactory.h b/Engine/Plugins/Experimental/Animation/SkeletonTemplateFramework/Source/Editor/Private/SkeletonTemplateFramework/Labels/LabelCollectionFactory.h
new file mode 100644
--- /dev/null
+++ b/Engine/Plugins/Experimental/Animation/SkeletonTemplateFramework/Source/Editor/Private/SkeletonTemplateFramework/Labels/LabelCollectionFactory.h
@@ -0,0 +1,23 @@
+// Copyright Epic Games, Inc. All Rights Reserved.
+
+#pragma once
+
+#include "CoreMinimal.h"
+#include "Factories/Factory.h"
+
+#include "LabelCollectionFactory.generated.h"
+
+UCLASS()
+class UAbstractSkeletonLabelCollectionFactory : public UFactory
+{
+	GENERATED_BODY()
+
+	UAbstractSkeletonLabelCollectionFactory();
+
+	virtual bool ConfigureProperties() override;
+	virtual UObject* FactoryCreateNew(UClass* Class, UObject* InParent, FName Name, EObjectFlags Flags, UObject* Context, FFeedbackContext* Warn, FName CallingContext) override;
+	virtual UObject* FactoryCreateNew(UClass* Class, UObject* InParent, FName Name, EObjectFlags Flags, UObject* Context, FFeedbackContext* Warn) override
+	{
+		return FactoryCreateNew(Class, InParent, Name, Flags, Context, Warn, NAME_None);
+	}
+};
diff --git a/Engine/Plugins/Experimental/Animation/SkeletonTemplateFramework/Source/Editor/Private/SkeletonTemplateFramework/Labels/LabelDragDrop.h b/Engine/Plugins/Experimental/Animation/SkeletonTemplateFramework/Source/Editor/Private/SkeletonTemplateFramework/Labels/LabelDragDrop.h
new file mode 100644
--- /dev/null
+++ b/Engine/Plugins/Experimental/Animation/SkeletonTemplateFramework/Source/Editor/Private/SkeletonTemplateFramework/Labels/LabelDragDrop.h
@@ -0,0 +1,42 @@
+// Copyright Epic Games, Inc. All Rights Reserved.
+
+#pragma once
+
+#include "CoreMinimal.h"
+#include "DragAndDrop/DecoratedDragDropOp.h"
+
+class UAbstractSkeletonLabelCollection;
+
+namespace UE::Anim::SkelTemplates::Labels
+{
+	class FLabelDragDropOp : public FDecoratedDragDropOp
+	{
+	public:
+		DRAG_DROP_OPERATOR_TYPE(FLabelDragDropOp, FDecoratedDragDropOp)
+
+		static TSharedRef<FLabelDragDropOp> New(const TWeakObjectPtr<const UAbstractSkeletonLabelCollection> InLabelCollection, const FName InLabel)
+		{
+			TSharedRef<FLabelDragDropOp> Operation = MakeShared<FLabelDragDropOp>();
+			Operation->LabelCollection = InLabelCollection;
+			Operation->Label = InLabel;
+
+			Operation->Construct();
+
+			return Operation;
+		}
+
+		TSharedPtr<SWidget> GetDefaultDecorator() const override
+		{
+			return SNew(SBorder)
+				.Visibility(EVisibility::Visible)
+				.BorderImage(FAppStyle::GetBrush("Menu.Background"))
+				[
+					SNew(STextBlock)
+						.Text(FText::FromName(Label))
+				];
+		}
+
+		TWeakObjectPtr<const UAbstractSkeletonLabelCollection> LabelCollection;
+		FName Label;
+	};
+}
\ No newline at end of file
diff --git a/Engine/Plugins/Experimental/Animation/SkeletonTemplateFramework/Source/Editor/Private/SkeletonTemplateFramework/Labels/SLabelBinding.cpp b/Engine/Plugins/Experimental/Animation/SkeletonTemplateFramework/Source/Editor/Private/SkeletonTemplateFramework/Labels/SLabelBinding.cpp
new file mode 100644
--- /dev/null
+++ b/Engine/Plugins/Experimental/Animation/SkeletonTemplateFramework/Source/Editor/Private/SkeletonTemplateFramework/Labels/SLabelBinding.cpp
@@ -0,0 +1,592 @@
+// Copyright Epic Games, Inc. All Rights Reserved.
+
+#include "SkeletonTemplateFramework/Labels/SLabelBinding.h"
+
+#include "Algo/Count.h"
+#include "AssetThumbnail.h"
+#include "Framework/Application/SlateApplication.h"
+#include "Framework/Commands/GenericCommands.h"
+#include "Framework/MultiBox/MultiBoxBuilder.h"
+#include "PropertyCustomizationHelpers.h"
+#include "ScopedTransaction.h"
+#include "SkeletonTemplateFramework/AbstractSkeletonLabelBinding.h"
+#include "SkeletonTemplateFramework/AbstractSkeletonLabelCollection.h"
+#include "SkeletonTemplateFramework/Labels/BoneDragDrop.h"
+#include "SkeletonTemplateFramework/Labels/LabelDragDrop.h"
+#include "SkeletonTemplateFramework/Labels/ILabelsTab.h"
+#include "SkeletonTemplateFramework/Labels/ILabelBinding.h"
+#include "SkeletonTemplateFramework/Labels/ILabelSkeletonTree.h"
+#include "SPositiveActionButton.h"
+#include "Widgets/Input/SHyperlink.h"
+#include "Widgets/Input/SSearchBox.h"
+#include "Widgets/Notifications/SProgressBar.h"
+#include "Widgets/Views/STreeView.h"
+
+#define LOCTEXT_NAMESPACE "UE::Anim::SkelTemplates::Labels::SLabelBinding"
+
+namespace UE::Anim::SkelTemplates::Labels
+{
+	void SLabelBinding::Construct(const FArguments& InArgs, TWeakObjectPtr<UAbstractSkeletonLabelBinding> InLabelBinding, TSharedPtr<ILabelsTab> InLabelsTabInterface)
+		{
+			LabelBinding = InLabelBinding;
+			OnTreeRefreshed = InArgs._OnTreeRefreshed;
+			LabelsTabInterface = InLabelsTabInterface;
+
+			ChildSlot
+			[
+				SNew(SVerticalBox)
+					+ SVerticalBox::Slot()
+					.AutoHeight()
+					.Padding(FMargin(1.0f, 1.0f))
+					[
+						SNew(SHorizontalBox)
+							+ SHorizontalBox::Slot()
+							.AutoWidth()
+							.Padding(FMargin(1.0f, 1.0f))
+							[
+								SNew(SPositiveActionButton)
+									.OnGetMenuContent(this, &SLabelBinding::CreateImportCollectionWidget)
+									.Icon(FAppStyle::Get().GetBrush("Icons.Import"))
+									.Text(LOCTEXT("AddButton_Text", "Import Collection"))
+							]
+							+ SHorizontalBox::Slot()
+							.FillWidth(1.0f)
+							.Padding(FMargin(1.0f, 1.0f))
+							[
+								SAssignNew(SearchBox, SSearchBox)
+									.SelectAllTextWhenFocused(true)
+									.OnTextChanged_Lambda([](const FText& InText) { /* TODO: Implement */ })
+									.HintText(LOCTEXT("SearchBox_Hint", "Search Labels..."))
+							]
+					]
+					+ SVerticalBox::Slot()
+					.FillHeight(1.0f)
+					.Padding(FMargin(1.0f, 1.0f))
+					[
+						SAssignNew(TreeView, STreeView<FTreeItemPtr>)
+							.SelectionMode(ESelectionMode::Single)
+							.HighlightParentNodesForSelection(true)
+							.TreeItemsSource(&RootItems)
+							.OnGenerateRow(this, &SLabelBinding::TreeView_OnGenerateRow)
+							.OnGetChildren(this, &SLabelBinding::TreeView_OnGetChildren)
+							.OnContextMenuOpening(this, &SLabelBinding::TreeView_OnContextMenuOpening)
+					]
+			];
+
+			RepopulateTreeData();
+			ExpandAllTreeItems();
+			BindCommands();
+		}
+
+		FReply SLabelBinding::OnKeyDown(const FGeometry& MyGeometry, const FKeyEvent& InKeyEvent)
+		{
+			if (CommandList->ProcessCommandBindings(InKeyEvent))
+			{
+				return FReply::Handled();
+			}
+
+			return FReply::Unhandled();
+		}
+
+		TSharedRef<ITableRow> SLabelBinding::TreeView_OnGenerateRow(FTreeItemPtr InItem, const TSharedRef<STableViewBase>& OwnerTable)
+		{
+			return InItem->GenerateRow(OwnerTable);
+		}
+
+		void SLabelBinding::TreeView_OnGetChildren(FTreeItemPtr InItem, TArray<FTreeItemPtr>& OutChildren)
+		{
+			InItem->GetChildren(OutChildren);
+		}
+
+		TSharedPtr<SWidget> SLabelBinding::TreeView_OnContextMenuOpening()
+		{
+			FMenuBuilder MenuBuilder(true, CommandList.ToSharedRef());
+
+			TArray<FTreeItemPtr> Selection = TreeView->GetSelectedItems();
+			if (Selection.Num() > 0)
+			{
+				FTreeItemPtr Selected = Selection[0]; // Single selection only
+
+				const bool bShowMenu = Selected->GenerateContextMenu(*this, MenuBuilder);
+
+				if (bShowMenu)
+				{
+					return MenuBuilder.MakeWidget();
+				}
+			}
+
+			return SNullWidget::NullWidget;
+		}
+
+		void SLabelBinding::PostUndo(bool bSuccess)
+		{
+			RepopulateTreeData();
+		}
+
+		void SLabelBinding::PostRedo(bool bSuccess)
+		{
+			RepopulateTreeData();
+		}
+
+		void SLabelBinding::ScrollToLabel(const TObjectPtr<const UAbstractSkeletonLabelCollection> InLabelCollection, const FName InLabel)
+		{
+			for (const FTreeItemPtr& RootItem : RootItems)
+			{
+				const TSharedPtr<FTreeItem_LabelCollection> LabelItem = StaticCastSharedPtr<FTreeItem_LabelCollection>(RootItem);
+				// All root items are label collection items
+
+				if (LabelItem->LabelCollection == InLabelCollection)
+				{
+					for (const TSharedPtr<FTreeItem_LabelBinding>& BindingItem : LabelItem->LabelBindings)
+					{
+						if (BindingItem->Label == InLabel)
+						{
+							TreeView->RequestScrollIntoView(BindingItem);
+							TreeView->SetSelection(BindingItem);
+							return;
+						}
+					}
+				}
+				
+			}
+		}
+
+		void SLabelBinding::RepopulateTreeData()
+		{
+			if (!bRepopulating)
+			{
+				TGuardValue<bool> RecursionGuard(bRepopulating, true);
+
+				// Rebuild tree data
+				RootItems.Reset();
+
+				if (LabelBinding.IsValid())
+				{
+					for (const TObjectPtr<const UAbstractSkeletonLabelCollection>& LabelCollection : LabelBinding->GetLabelCollections())
+					{
+						TSharedPtr<FTreeItem_LabelCollection> LabelCollectionItem = MakeShared<FTreeItem_LabelCollection>();
+						LabelCollectionItem->LabelCollection = LabelCollection;
+						LabelCollectionItem->LabelsTabInterface = LabelsTabInterface;
+
+						for (const FName Label : LabelCollection->GetLabels())
+						{
+							TSharedPtr<FTreeItem_LabelBinding> LabelBindingItem = MakeShared<FTreeItem_LabelBinding>();
+							LabelBindingItem->Label = Label;
+							LabelBindingItem->BoneName = LabelBinding->GetLabelBinding(LabelCollection, Label);
+							LabelBindingItem->LabelCollection = LabelCollection;
+							LabelBindingItem->LabelsTabInterface = LabelsTabInterface;
+
+							LabelCollectionItem->LabelBindings.Add(LabelBindingItem);
+						}
+
+						RootItems.Add(LabelCollectionItem);
+					}
+				}
+
+				TreeView->RequestTreeRefresh();
+
+				for (FTreeItemPtr TreeItem : GetAllTreeItems())
+				{
+					TreeView->SetItemExpansion(TreeItem, true);
+				}
+
+				OnTreeRefreshed.ExecuteIfBound();
+			}
+		}
+
+		TWeakObjectPtr<UAbstractSkeletonLabelBinding> SLabelBinding::GetLabelBinding() const
+		{
+			return LabelBinding;
+		}
+
+		void SLabelBinding::ExpandAllTreeItems()
+		{
+			for (FTreeItemPtr TreeItem : GetAllTreeItems())
+			{
+				TreeView->SetItemExpansion(TreeItem, true);
+			}
+		}
+
+		void SLabelBinding::BindCommands()
+		{
+			CommandList = MakeShareable(new FUICommandList);
+
+			CommandList->MapAction(
+				FGenericCommands::Get().Delete,
+				FExecuteAction::CreateSP(this, &SLabelBinding::OnRemoveSelected)
+			);
+		}
+
+		TArray<TSharedPtr<SLabelBinding::ITreeItem>> SLabelBinding::GetAllTreeItems()
+		{
+			TArray<FTreeItemPtr> AllItems;
+			AllItems.Append(RootItems);
+
+			for (int32 Index = 0; Index < AllItems.Num(); ++Index)
+			{
+				AllItems[Index]->GetChildren(AllItems);
+			}
+
+			return AllItems;
+		}
+
+		void SLabelBinding::OnImportLabelCollection()
+		{
+		
+		}
+
+		TSharedRef<SWidget> SLabelBinding::CreateImportCollectionWidget()
+		{
+			FMenuBuilder MenuBuilder(true, nullptr);
+
+			const bool bAllowClear = false;
+			const bool bAllowCopyPaste = false;
+
+			TArray<const UClass*> AllowedClasses;
+			AllowedClasses.Add(UAbstractSkeletonLabelCollection::StaticClass());
+
+			TSharedRef<SWidget> AssetPickerWidget = PropertyCustomizationHelpers::MakeAssetPickerWithMenu(
+				FAssetData(),
+				bAllowClear,
+				bAllowCopyPaste,
+				AllowedClasses,
+				PropertyCustomizationHelpers::GetNewAssetFactoriesForClasses(AllowedClasses),
+				FOnShouldFilterAsset::CreateLambda([this](const FAssetData& InAssetData)
+				{
+					for (FTreeItemPtr TreeItem : RootItems)
+					{
+						const TSharedPtr<FTreeItem_LabelCollection> LabelCollectionItem = StaticCastSharedPtr<FTreeItem_LabelCollection>(TreeItem);
+						// All root items are label collection items
+
+						if (FAssetData(LabelCollectionItem->LabelCollection.Get()) == InAssetData)
+						{
+							return true;
+						}
+					}
+
+					return false;
+				}),
+				FOnAssetSelected::CreateLambda([this](const FAssetData& InAssetData)
+				{
+					const FScopedTransaction Transaction(LOCTEXT("AddLabelCollection", "Add Label Collection"));
+					LabelBinding->Modify();
+
+					TObjectPtr<UAbstractSkeletonLabelCollection> SetCollection = CastChecked<UAbstractSkeletonLabelCollection>(InAssetData.GetAsset());
+
+					const bool bSuccess = LabelBinding->AddLabelCollection(SetCollection);
+					check(bSuccess);
+
+					LabelsTabInterface->RepopulateLabelData();
+				}),
+				FSimpleDelegate::CreateLambda([]()
+				{
+					FSlateApplication::Get().DismissAllMenus();
+				}));
+
+			MenuBuilder.AddWidget(AssetPickerWidget, FText::GetEmpty(), /*bNoIndent=*/ true);
+
+			return MenuBuilder.MakeWidget();
+		}
+
+		void SLabelBinding::OnRemoveSelected()
+		{
+			TArray<FTreeItemPtr> Selection = TreeView->GetSelectedItems();
+			
+			if (ensure(Selection.Num() == 1))
+			{
+				FTreeItemPtr Selected = Selection[0];
+				Selected->OnRemove();
+			}
+		}
+
+		// FTreeItem_LabelCollection
+
+		void SLabelBinding::FTreeItem_LabelCollection::GetChildren(TArray<FTreeItemPtr>& OutChildren)
+		{
+			OutChildren.Append(LabelBindings);
+		}
+
+		TSharedRef<ITableRow> SLabelBinding::FTreeItem_LabelCollection::GenerateRow(const TSharedRef<STableViewBase>& OwnerTable)
+		{
+			const int32 NumBound = Algo::CountIf(LabelBindings, [](const TSharedPtr<FTreeItem_LabelBinding> InBinding) -> bool
+				{
+					return InBinding->BoneName != NAME_None;
+				});
+			const int32 NumUnbound = LabelBindings.Num() - NumBound;
+
+			return SNew(STableRow<FTreeItemPtr>, OwnerTable)
+				[
+					SNew(SHorizontalBox)
+						+ SHorizontalBox::Slot()
+						.FillWidth(0.5)
+						.VAlign(VAlign_Center)
+						.Padding(FMargin(4.0f))
+						[
+							SNew(SBox)
+								.HeightOverride(8.0f)
+								[
+									SNew(SProgressBar)
+										.FillColorAndOpacity(NumUnbound == 0 ? FAppStyle::Get().GetSlateColor("Colors.AccentGreen") : FAppStyle::Get().GetSlateColor("Colors.AccentBlue"))
+										.FillImage(FAppStyle::Get().GetBrush("WhiteBrush"))
+										.Percent(static_cast<float>(NumBound) / LabelBindings.Num())
+										.RefreshRate(0)
+								]
+						]
+						+ SHorizontalBox::Slot()
+						.FillWidth(0.5)
+						.Padding(FMargin(4.0f))
+						[
+							SNew(SObjectPropertyEntryBox)
+								.AllowedClass(UAbstractSkeletonLabelCollection::StaticClass())
+								.ObjectPath_Lambda([&]()
+								{
+									return LabelCollection->GetPathName();
+								})
+								.OnObjectChanged_Lambda([](const FAssetData& AssetData)
+								{
+									// TODO: Implement
+								})
+								.AllowClear(true)
+								.DisplayUseSelected(true)
+								.DisplayBrowse(true)
+								.DisplayThumbnail(true)
+								.ThumbnailPool(MakeShared<FAssetThumbnailPool>(1))
+						]
+				];
+		}
+
+		bool SLabelBinding::FTreeItem_LabelCollection::GenerateContextMenu(SLabelBinding& LabelBindingWidget, FMenuBuilder& MenuBuilder)
+		{
+			const bool bHasBoundChildren = LabelBindings.ContainsByPredicate([](const TSharedPtr<FTreeItem_LabelBinding> InBinding)
+			{
+				return InBinding->BoneName != NAME_None;
+			});
+
+			if (bHasBoundChildren)
+			{
+				MenuBuilder.AddMenuEntry(
+					LOCTEXT("UnbindAll_Label", "Unbind All"),
+					LOCTEXT("UnbindAll_Tooltip", "Unbind all bones and labels"),
+					FSlateIcon(FAppStyle::GetAppStyleSetName(), "Icons.Unlink"),
+					FUIAction(
+						FExecuteAction::CreateLambda([this]()
+							{
+								const FScopedTransaction Transaction(LOCTEXT("UnbindAll", "Unbind All"));
+								LabelsTabInterface->GetLabelBindingWidget()->GetLabelBinding()->Modify();
+
+								for (const TSharedPtr<FTreeItem_LabelBinding>& Binding : LabelBindings)
+								{
+									if (Binding->BoneName != NAME_None)
+									{
+										const bool bSuccess = LabelsTabInterface->GetLabelBindingWidget()->GetLabelBinding()->UnbindBoneFromLabel(Binding->LabelCollection.Get(), Binding->BoneName, Binding->Label);
+										check(bSuccess);
+									}
+								}
+
+								LabelsTabInterface->GetLabelBindingWidget()->RepopulateTreeData();
+								LabelsTabInterface->GetLabelSkeletonTreeWidget()->RepopulateTreeData();
+							}
+						)
+					));
+			}
+
+			MenuBuilder.AddSeparator();
+
+			MenuBuilder.AddMenuEntry(FGenericCommands::Get().Delete, NAME_None, LOCTEXT("Remove", "Remove"));
+
+			return true;
+		}
+
+		void SLabelBinding::FTreeItem_LabelCollection::OnRemove()
+		{
+			TSharedPtr<ILabelBindingWidget> LabelBindingInterface = LabelsTabInterface->GetLabelBindingWidget();
+
+			const FScopedTransaction Transaction(LOCTEXT("RemoveLabelCollection", "Remove Label Collection"));
+			LabelBindingInterface->GetLabelBinding()->Modify();
+
+			const bool bSuccess = LabelBindingInterface->GetLabelBinding()->RemoveLabelCollection(LabelCollection.Get());
+			check(bSuccess);
+
+			LabelBindingInterface->RepopulateTreeData();
+			LabelsTabInterface->GetLabelSkeletonTreeWidget()->RepopulateTreeData();
+		}
+
+		// FTreeItem_LabelBinding
+
+		void SLabelBinding::FTreeItem_LabelBinding::GetChildren(TArray<FTreeItemPtr>& OutChildren)
+		{
+			// Do nothing
+		}
+
+		TSharedRef<ITableRow> SLabelBinding::FTreeItem_LabelBinding::GenerateRow(const TSharedRef<STableViewBase>& OwnerTable)
+		{
+			TSharedPtr<SWidget> BoneWidget;
+
+			if (BoneName == NAME_None)
+			{
+				BoneWidget = SNew(STextBlock)
+					.ColorAndOpacity(FSlateColor::UseSubduedForeground())
+					.Text(FText::FromName(NAME_None));
+			}
+			else
+			{
+				BoneWidget = SNew(SHyperlink)
+					.Text(FText::FromName(BoneName))
+					.OnNavigate_Lambda([this]()
+						{
+							TSharedPtr<ILabelSkeletonTreeWidget> LabelSkeletonTreeWidget = LabelsTabInterface->GetLabelSkeletonTreeWidget();
+							LabelSkeletonTreeWidget->ScrollToBone(BoneName);
+						});
+			}
+
+			return SNew(STableRow<FTreeItemPtr>, OwnerTable)
+				.OnDragDetected_Lambda([this](const FGeometry& MyGeometry, const FPointerEvent& MouseEvent)
+					{
+						if (MouseEvent.IsMouseButtonDown(EKeys::LeftMouseButton))
+						{
+							const TSharedRef<FLabelDragDropOp> DragDropOp = FLabelDragDropOp::New(LabelCollection, Label);
+							return FReply::Handled().BeginDragDrop(DragDropOp);
+						}
+
+						return FReply::Unhandled();
+					})
+				.OnDragEnter_Lambda([this](const FDragDropEvent& DragDropEvent)
+				{
+					const TSharedPtr<FBoneDragDropOp> DragDropOp = DragDropEvent.GetOperationAs<FBoneDragDropOp>();
+					if (DragDropOp.IsValid())
+					{
+						DragDropOp->SetHoveredLabel(Label);
+					}
+				})
+				.OnDragLeave_Lambda([this](const FDragDropEvent& DragDropEvent)
+				{
+					const TSharedPtr<FBoneDragDropOp> DragDropOp = DragDropEvent.GetOperationAs<FBoneDragDropOp>();
+					if (DragDropOp.IsValid())
+					{
+						DragDropOp->SetHoveredLabel(NAME_None);
+					}
+				})
+				.OnCanAcceptDrop_Lambda([](const FDragDropEvent& DragDropEvent, EItemDropZone DropZone, FTreeItemPtr TargetTreeItem) -> TOptional<EItemDropZone>
+					{
+						TOptional<EItemDropZone> ReturnedDropZone;
+
+						const TSharedPtr<FBoneDragDropOp> DragDropOp = DragDropEvent.GetOperationAs<FBoneDragDropOp>();
+						if (DragDropOp.IsValid())
+						{
+							ReturnedDropZone = EItemDropZone::OntoItem;
+						}
+
+						return ReturnedDropZone;
+					})
+				.OnAcceptDrop_Lambda([this](const FDragDropEvent& DragDropEvent, EItemDropZone DropZone, FTreeItemPtr TargetTreeItem) -> FReply
+					{
+						const TSharedPtr<FBoneDragDropOp> DragDropOp = DragDropEvent.GetOperationAs<FBoneDragDropOp>();
+
+						if (DragDropOp.IsValid())
+						{
+							const FScopedTransaction Transaction(LOCTEXT("BindBoneToLabel", "Bind Bone to Label"));
+							DragDropOp->LabelBinding->Modify();
+
+							const bool bSuccess = DragDropOp->LabelBinding->BindBoneToLabel(LabelCollection.Get(), DragDropOp->BoneName, Label);
+							check(bSuccess);
+
+							LabelsTabInterface->GetLabelSkeletonTreeWidget()->RepopulateTreeData();
+							LabelsTabInterface->GetLabelBindingWidget()->RepopulateTreeData();
+
+							return FReply::Handled();
+						}
+
+						return FReply::Unhandled();
+					})
+				.ShowWires(true)
+				.Style(FAppStyle::Get(), "TableView.AlternatingRow")
+				[
+					SNew(SHorizontalBox)
+						+ SHorizontalBox::Slot()
+						.FillWidth(0.5)
+						[
+							SNew(SHorizontalBox)
+								+ SHorizontalBox::Slot()
+								.AutoWidth()
+								.Padding(FMargin(4.0f))
+								[
+									SNew(SImage)
+										.ColorAndOpacity(BoneName == NAME_None ? FSlateColor::UseForeground() : FAppStyle::Get().GetSlateColor("Colors.AccentBlue"))
+										.Image(FAppStyle::Get().GetBrush("Icons.Tag"))
+								]
+								+ SHorizontalBox::Slot()
+								.AutoWidth()
+								.Padding(FMargin(4.0f))
+								[
+									SNew(STextBlock)
+										.Text(FText::FromName(Label))
+								]
+						]
+						+ SHorizontalBox::Slot()
+						.FillWidth(0.5)
+						[
+							SNew(SHorizontalBox)
+								+ SHorizontalBox::Slot()
+								.AutoWidth()
+								.Padding(FMargin(4.0f))
+								[
+									SNew(SImage)
+										.ColorAndOpacity(BoneName == NAME_None ? FSlateColor::UseSubduedForeground() : FAppStyle::Get().GetSlateColor("Colors.AccentBlue"))
+										.Image(FAppStyle::Get().GetBrush("SkeletonTree.Bone"))
+								]
+								+ SHorizontalBox::Slot()
+								.AutoWidth()
+								.Padding(FMargin(4.0f))
+								[
+									BoneWidget.ToSharedRef()
+								]
+						]
+				];
+		}
+
+		bool SLabelBinding::FTreeItem_LabelBinding::GenerateContextMenu(SLabelBinding& InLabelBindingWidget, FMenuBuilder& MenuBuilder)
+		{
+			if (BoneName != NAME_None)
+			{
+				MenuBuilder.AddMenuEntry(
+					FText::Format(LOCTEXT("UnbindBone_Label", "Unbind {0}"), FText::FromName(BoneName)),
+					LOCTEXT("UnbindBone_Tooltip", "Unbind the selected bone and label"),
+					FSlateIcon(FAppStyle::GetAppStyleSetName(), "Icons.Unlink"),
+					FUIAction(
+						FExecuteAction::CreateLambda([this]()
+							{
+								const bool bSuccess = LabelsTabInterface->GetLabelBindingWidget()->GetLabelBinding()->UnbindBoneFromLabel(LabelCollection.Get(), BoneName, Label);
+								check(bSuccess);
+
+								LabelsTabInterface->GetLabelBindingWidget()->RepopulateTreeData();
+								LabelsTabInterface->GetLabelSkeletonTreeWidget()->RepopulateTreeData();
+							}
+						)
+					));
+
+				return true;
+			}
+
+			return false;
+		}
+
+		void SLabelBinding::FTreeItem_LabelBinding::OnRemove()
+		{
+			if (Label != NAME_None)
+			{
+				TSharedPtr<ILabelBindingWidget> LabelBindingInterface = LabelsTabInterface->GetLabelBindingWidget();
+
+				const FScopedTransaction Transaction(LOCTEXT("UnbindLabel", "Unbind Label"));
+				LabelBindingInterface->GetLabelBinding()->Modify();
+
+				const bool bSuccess = LabelBindingInterface->GetLabelBinding()->UnbindBoneFromLabel(LabelCollection.Get(), BoneName, Label);
+				check(bSuccess);
+
+				LabelBindingInterface->RepopulateTreeData();
+				LabelsTabInterface->GetLabelSkeletonTreeWidget()->RepopulateTreeData();
+			}
+		}
+
+}
+
+#undef LOCTEXT_NAMESPACE
\ No newline at end of file
diff --git a/Engine/Plugins/Experimental/Animation/SkeletonTemplateFramework/Source/Editor/Private/SkeletonTemplateFramework/Labels/SLabelBinding.h b/Engine/Plugins/Experimental/Animation/SkeletonTemplateFramework/Source/Editor/Private/SkeletonTemplateFramework/Labels/SLabelBinding.h
new file mode 100644
--- /dev/null
+++ b/Engine/Plugins/Experimental/Animation/SkeletonTemplateFramework/Source/Editor/Private/SkeletonTemplateFramework/Labels/SLabelBinding.h
@@ -0,0 +1,118 @@
+// Copyright Epic Games, Inc. All Rights Reserved.
+
+#pragma once
+
+#include "CoreMinimal.h"
+#include "EditorUndoClient.h"
+#include "SkeletonTemplateFramework/Labels/ILabelBinding.h"
+#include "Widgets/SCompoundWidget.h"
+
+class FMenuBuilder;
+class FUICommandList;
+class ITableRow;
+class SSearchBox;
+class STableViewBase;
+template<typename> class STreeView;
+class UAbstractSkeletonLabelBinding;
+class UAbstractSkeletonLabelCollection;
+
+namespace UE::Anim::SkelTemplates::Labels
+{
+	class ILabelsTab;
+
+	class SLabelBinding : public SCompoundWidget, public FSelfRegisteringEditorUndoClient, public ILabelBindingWidget
+	{
+		SLATE_BEGIN_ARGS(SLabelBinding) {}
+			SLATE_EVENT(FSimpleDelegate, OnTreeRefreshed)
+		SLATE_END_ARGS()
+
+	public:
+		struct ITreeItem
+		{
+			virtual ~ITreeItem() = default;
+			virtual void GetChildren(TArray<TSharedPtr<ITreeItem>>& OutChildren) = 0;
+			virtual TSharedRef<ITableRow> GenerateRow(const TSharedRef<STableViewBase>& OwnerTable) = 0;
+			virtual bool GenerateContextMenu(SLabelBinding& SetBindingWidget, FMenuBuilder& Menu) = 0;
+			virtual void OnRemove() = 0;
+		};
+
+		using FTreeItemPtr = TSharedPtr<ITreeItem>;
+
+		struct FTreeItem_LabelBinding : public ITreeItem
+		{
+			virtual ~FTreeItem_LabelBinding() override = default;
+			virtual void GetChildren(TArray<FTreeItemPtr>& OutChildren) override;
+			virtual TSharedRef<ITableRow> GenerateRow(const TSharedRef<STableViewBase>& OwnerTable) override;
+			virtual bool GenerateContextMenu(SLabelBinding& SetBindingWidget, FMenuBuilder& Menu) override;
+			virtual void OnRemove() override;
+
+			FName Label;
+			FName BoneName;
+			TWeakObjectPtr<const UAbstractSkeletonLabelCollection> LabelCollection;
+			TSharedPtr<ILabelsTab> LabelsTabInterface;
+		};
+
+		struct FTreeItem_LabelCollection : public ITreeItem
+		{
+			virtual ~FTreeItem_LabelCollection() override = default;
+			virtual void GetChildren(TArray<FTreeItemPtr>& OutChildren) override;
+			virtual TSharedRef<ITableRow> GenerateRow(const TSharedRef<STableViewBase>& OwnerTable) override;
+			virtual bool GenerateContextMenu(SLabelBinding& SetBindingWidget, FMenuBuilder& Menu) override;
+			virtual void OnRemove() override;
+
+			TWeakObjectPtr<const UAbstractSkeletonLabelCollection> LabelCollection;
+			TArray<TSharedPtr<FTreeItem_LabelBinding>> LabelBindings;
+			TSharedPtr<ILabelsTab> LabelsTabInterface;
+		};
+
+		void Construct(const FArguments& InArgs, TWeakObjectPtr<UAbstractSkeletonLabelBinding> InSetBinding, TSharedPtr<ILabelsTab> InLabelsTabInterface);
+
+		/** FSelfRegisteringEditorUndoClient interface */
+		virtual void PostUndo(bool bSuccess) override;
+		virtual void PostRedo(bool bSuccess) override;
+
+		/** ILabelBindingWidget */
+		virtual void ScrollToLabel(const TObjectPtr<const UAbstractSkeletonLabelCollection> InLabelCollection, const FName InLabel) override;
+		virtual void RepopulateTreeData() override;
+		virtual TWeakObjectPtr<UAbstractSkeletonLabelBinding> GetLabelBinding() const override;
+
+	private:
+		virtual FReply OnKeyDown(const FGeometry& MyGeometry, const FKeyEvent& InKeyEvent) override;
+
+		TSharedRef<ITableRow> TreeView_OnGenerateRow(FTreeItemPtr InItem, const TSharedRef<STableViewBase>& OwnerTable);
+
+		void TreeView_OnGetChildren(FTreeItemPtr InItem, TArray<FTreeItemPtr>& OutChildren);
+
+		TSharedPtr<SWidget> TreeView_OnContextMenuOpening();
+
+		void ExpandAllTreeItems();
+
+		void BindCommands();
+
+		TArray<FTreeItemPtr> GetAllTreeItems();
+
+		void OnImportLabelCollection();
+
+		TSharedRef<SWidget> CreateImportCollectionWidget();
+
+		void OnRemoveSelected();
+
+		TSharedPtr<ILabelsTab> LabelsTabInterface;
+
+	public:
+		TWeakObjectPtr<UAbstractSkeletonLabelBinding> LabelBinding;
+
+	private:
+		TSharedPtr<STreeView<FTreeItemPtr>> TreeView;
+
+		TArray<FTreeItemPtr> RootItems;
+
+		TSharedPtr<SSearchBox> SearchBox;
+
+		TSharedPtr<FUICommandList> CommandList;
+
+		FSimpleDelegate OnTreeRefreshed;
+
+		bool bRepopulating = false;
+	};
+}
\ No newline at end of file
diff --git a/Engine/Plugins/Experimental/Animation/SkeletonTemplateFramework/Source/Editor/Private/SkeletonTemplateFramework/Labels/SLabelCollection.cpp b/Engine/Plugins/Experimental/Animation/SkeletonTemplateFramework/Source/Editor/Private/SkeletonTemplateFramework/Labels/SLabelCollection.cpp
new file mode 100644
--- /dev/null
+++ b/Engine/Plugins/Experimental/Animation/SkeletonTemplateFramework/Source/Editor/Private/SkeletonTemplateFramework/Labels/SLabelCollection.cpp
@@ -0,0 +1,555 @@
+// Copyright Epic Games, Inc. All Rights Reserved.
+
+#include "SkeletonTemplateFramework/Labels/SLabelCollection.h"
+
+#include "Framework/Commands/GenericCommands.h"
+#include "Framework/MultiBox/MultiBoxBuilder.h"
+#include "ScopedTransaction.h"
+#include "SkeletonTemplateFramework/AbstractSkeletonLabelCollection.h"
+#include "SPositiveActionButton.h"
+#include "Widgets/Input/SSearchBox.h"
+#include "Widgets/Text/SInlineEditableTextBlock.h"
+#include "DragAndDrop/DecoratedDragDropOp.h"
+#include "Widgets/Text/SRichTextBlock.h"
+#include "Widgets/Views/SListView.h"
+
+#define LOCTEXT_NAMESPACE "UE::Anim::SkelTemplates::Labels::SLabelCollection"
+
+namespace UE::Anim::SkelTemplates::Labels
+{
+	class FLabelCollectionDragDropOp : public FDecoratedDragDropOp
+	{
+	public:
+		DRAG_DROP_OPERATOR_TYPE(FLabelCollectionDragDropOp, FDecoratedDragDropOp)
+
+		enum class EAction
+		{
+			None,
+			SwapLabels,
+			MoveLabelAbove,
+			MoveLabelBelow
+		};
+
+		static TSharedRef<FLabelCollectionDragDropOp> New(const FName InLabel)
+		{
+			TSharedRef<FLabelCollectionDragDropOp> Operation = MakeShared<FLabelCollectionDragDropOp>();
+			Operation->Label = InLabel;
+			Operation->Construct();
+
+			return Operation;
+		}
+
+		TSharedPtr<SWidget> GetDefaultDecorator() const override
+		{
+			TSharedRef<SHorizontalBox> VBox = SNew(SHorizontalBox);
+
+			// EAction::None
+			VBox->AddSlot()
+				[
+					SNew(SHorizontalBox)
+					.Visibility_Lambda([this]()
+					{
+						return Action == EAction::None ? EVisibility::Visible : EVisibility::Collapsed;
+					})
+					+ SHorizontalBox::Slot()
+					.AutoWidth()
+					[
+						SNew(STextBlock)
+							.Text_Lambda([this]()
+							{
+								return FText::FromName(Label);
+							})
+					]
+				];
+
+			// EAction::SwapLabels
+			VBox->AddSlot()
+				[
+					SNew(SHorizontalBox)
+					.Visibility_Lambda([this]()
+					{
+						return Action == EAction::SwapLabels ? EVisibility::Visible : EVisibility::Collapsed;
+					})
+					+ SHorizontalBox::Slot()
+					.AutoWidth()
+					[
+						SNew(SRichTextBlock)
+							.DecoratorStyleSet(&FAppStyle::Get())
+							.Text_Lambda([this]()
+								{
+									return FText::Format(LOCTEXT("SwapLabels", "Swap <RichTextBlock.Bold>{0}</> with <RichTextBlock.Bold>{1}</>"), FText::FromName(Label), ActionLabel);
+								})
+					]
+				];
+			
+			// EAction::SwapLabels
+			VBox->AddSlot()
+				[
+					SNew(SHorizontalBox)
+					.Visibility_Lambda([this]()
+					{
+						return Action == EAction::MoveLabelAbove ? EVisibility::Visible : EVisibility::Collapsed;
+					})
+					+ SHorizontalBox::Slot()
+					.AutoWidth()
+					[
+						SNew(SRichTextBlock)
+							.DecoratorStyleSet(&FAppStyle::Get())
+							.Text_Lambda([this]()
+								{
+									return FText::Format(LOCTEXT("MoveAbove", "Move <RichTextBlock.Bold>{0}</> above <RichTextBlock.Bold>{1}</>"), FText::FromName(Label), ActionLabel);
+								})
+					]
+				];
+			
+			// EAction::SwapLabels
+			VBox->AddSlot()
+				[
+					SNew(SHorizontalBox)
+					.Visibility_Lambda([this]()
+					{
+						return Action == EAction::MoveLabelBelow ? EVisibility::Visible : EVisibility::Collapsed;
+					})
+					+ SHorizontalBox::Slot()
+					.Padding(3.0f)
+					.AutoWidth()
+					[
+						SNew(SRichTextBlock)
+							.DecoratorStyleSet(&FAppStyle::Get())
+							.Text_Lambda([this]()
+								{
+									return FText::Format(LOCTEXT("MoveBelow", "Move <RichTextBlock.Bold>{0}</> below <RichTextBlock.Bold>{1}</>"), FText::FromName(Label), ActionLabel);
+								})
+					]
+				];
+			
+			return SNew(SBorder)
+				.Visibility(EVisibility::Visible)
+				.BorderImage(FAppStyle::GetBrush("Menu.Background"))
+				.Padding(FMargin(8.0f, 4.0f))
+				[
+					VBox
+				];
+		}
+		
+		void SetAction(const EAction InAction, const FText InLabel)
+		{
+			ActionLabel = InLabel;
+			Action = InAction;
+		}
+
+		void ClearAction()
+		{
+			ActionLabel = FText();
+			Action = EAction::None;
+		}
+		
+		EAction Action = EAction::None;
+		
+		FText ActionLabel;
+
+		FName Label;
+	};
+
+	void SLabelCollection::Construct(const FArguments& InArgs, TWeakObjectPtr<UAbstractSkeletonLabelCollection> InLabelCollection)
+	{
+		LabelCollection = InLabelCollection;
+
+		ChildSlot
+		[
+			SNew(SVerticalBox)
+				+ SVerticalBox::Slot()
+				.AutoHeight()
+				.Padding(FMargin(1.0f, 1.0f))
+				[
+					SNew(SHorizontalBox)
+						+ SHorizontalBox::Slot()
+						.AutoWidth()
+						.Padding(FMargin(1.0f, 1.0f))
+						[
+							SNew(SPositiveActionButton)
+								.OnClicked_Lambda([this]()
+								{
+									OnAddLabel();
+									return FReply::Handled();
+								})
+								.Icon(FAppStyle::Get().GetBrush("Icons.Plus"))
+								.Text(LOCTEXT("AddButton_Text", "Add Label"))
+						]
+						+ SHorizontalBox::Slot()
+						.FillWidth(1.0f)
+						.Padding(FMargin(1.0f, 1.0f))
+						[
+							SAssignNew(SearchBox, SSearchBox)
+								.SelectAllTextWhenFocused(true)
+								.OnTextChanged_Lambda([this](const FText& InText)
+								{
+									SearchText = InText;
+									ApplySearchFilter();
+								})
+								.HintText(LOCTEXT("SearchBox_Hint", "Search Labels..."))
+						]
+				]
+				+ SVerticalBox::Slot()
+				.FillHeight(1.0f)
+				.Padding(FMargin(1.0f, 1.0f))
+				[
+					SAssignNew(ListView, SListView<FListItemPtr>)
+						.SelectionMode(ESelectionMode::Single)
+						.ListItemsSource(&FilteredItems)
+						.OnGenerateRow(this, &SLabelCollection::ListView_OnGenerateRow)
+						.OnContextMenuOpening(this, &SLabelCollection::ListView_OnContextMenuOpening)
+						.OnItemScrolledIntoView(this, &SLabelCollection::ListView_OnItemScrolledIntoView)
+						.OnMouseButtonDoubleClick(this, &SLabelCollection::ListView_OnMouseButtonDoubleClick)
+						.HeaderRow
+						(
+							SNew(SHeaderRow)
+								+ SHeaderRow::Column(FName("Labels"))
+								.FillWidth(0.5f)
+								.DefaultLabel(LOCTEXT("ListView_Header", "Labels"))
+						)
+				]
+		];
+
+		RepopulateListData();
+		BindCommands();
+	}
+
+	void SLabelCollection::PostUndo(bool bSuccess)
+	{
+		RepopulateListData();
+	}
+
+	void SLabelCollection::PostRedo(bool bSuccess)
+	{
+		RepopulateListData();
+	}
+
+	FReply SLabelCollection::OnKeyDown(const FGeometry& MyGeometry, const FKeyEvent& InKeyEvent)
+	{
+		if (CommandList->ProcessCommandBindings(InKeyEvent))
+		{
+			return FReply::Handled();
+		}
+
+		return FReply::Unhandled();
+	}
+
+	TSharedRef<ITableRow> SLabelCollection::ListView_OnGenerateRow(FListItemPtr InItem, const TSharedRef<STableViewBase>& OwnerTable)
+	{
+		TSharedPtr<SInlineEditableTextBlock> InlineWidget;
+
+		TSharedRef<STableRow<FListItemPtr>> RowWidget = SNew(STableRow<FListItemPtr>, OwnerTable)
+			.Style(&FAppStyle().GetWidgetStyle<FTableRowStyle>("TableView.AlternatingRow"))
+			.OnDragLeave_Lambda([this](const FDragDropEvent& DragDropEvent)
+			{
+				const TSharedPtr<FLabelCollectionDragDropOp> DragDropOp = DragDropEvent.GetOperationAs<FLabelCollectionDragDropOp>();
+				if (DragDropOp.IsValid())
+				{
+					DragDropOp->ClearAction();
+				}
+			})
+			.OnDragDetected_Lambda([InItem](const FGeometry& MyGeometry, const FPointerEvent& MouseEvent)
+				{
+					if (MouseEvent.IsMouseButtonDown(EKeys::LeftMouseButton))
+					{
+						const TSharedRef<FLabelCollectionDragDropOp> DragDropOp = FLabelCollectionDragDropOp::New(InItem->Label);
+						return FReply::Handled().BeginDragDrop(DragDropOp);
+					}
+
+					return FReply::Unhandled();
+				})
+			.OnCanAcceptDrop_Lambda([](const FDragDropEvent& DragDropEvent, EItemDropZone DropZone, FListItemPtr TargetTreeItem) -> TOptional<EItemDropZone>
+				{
+					using namespace UE::Anim::SkelTemplates;
+
+					TOptional<EItemDropZone> ReturnedDropZone;
+
+					const TSharedPtr<FLabelCollectionDragDropOp> DragDropOp = DragDropEvent.GetOperationAs<FLabelCollectionDragDropOp>();
+					if (DragDropOp.IsValid() && DragDropOp->Label != TargetTreeItem->Label)
+					{
+						if (DropZone == EItemDropZone::OntoItem)
+						{
+							DragDropOp->SetAction(FLabelCollectionDragDropOp::EAction::SwapLabels, FText::FromName(TargetTreeItem->Label));
+						}
+						else if (DropZone == EItemDropZone::AboveItem)
+						{
+							DragDropOp->SetAction(FLabelCollectionDragDropOp::EAction::MoveLabelAbove, FText::FromName(TargetTreeItem->Label));
+						}
+						else if (DropZone == EItemDropZone::BelowItem)
+						{
+							DragDropOp->SetAction(FLabelCollectionDragDropOp::EAction::MoveLabelBelow, FText::FromName(TargetTreeItem->Label));
+						}
+						
+						ReturnedDropZone = DropZone;
+					}
+
+					return ReturnedDropZone;
+				})
+			.OnAcceptDrop_Lambda([this](const FDragDropEvent& DragDropEvent, EItemDropZone DropZone, FListItemPtr TargetTreeItem) -> FReply
+				{
+					const TSharedPtr<FLabelCollectionDragDropOp> DragDropOp = DragDropEvent.GetOperationAs<FLabelCollectionDragDropOp>();
+					if (DragDropOp.IsValid())
+					{
+						const FScopedTransaction Transaction(LOCTEXT("MoveLabel", "Move Label"));
+						LabelCollection->Modify();
+
+
+						if (DropZone == EItemDropZone::OntoItem)
+						{
+							TArray<FName>& Labels = LabelCollection->GetMutableLabels();
+
+							const int32 A = Labels.Find(TargetTreeItem->Label);
+							const int32 B = Labels.Find(DragDropOp->Label);
+						
+							// Swap labels
+							if (ensure(A != INDEX_NONE && B != INDEX_NONE))
+							{
+								Labels[A] = DragDropOp->Label;
+								Labels[B] = TargetTreeItem->Label;
+							}
+						}
+						else if (DropZone == EItemDropZone::BelowItem)
+						{
+							LabelCollection->RemoveLabel(DragDropOp->Label);
+
+							TArray<FName>& Labels = LabelCollection->GetMutableLabels();
+
+							const int32 A = Labels.Find(TargetTreeItem->Label);
+
+							Labels.Insert(DragDropOp->Label, A + 1);
+						}
+						else
+						{
+							LabelCollection->RemoveLabel(DragDropOp->Label);
+
+							TArray<FName>& Labels = LabelCollection->GetMutableLabels();
+
+							const int32 A = Labels.Find(TargetTreeItem->Label);
+
+							Labels.Insert(DragDropOp->Label, A);
+						}
+
+						RepopulateListData();
+
+						return FReply::Handled();
+					}
+
+					return FReply::Unhandled();
+				})
+			[
+				SNew(SHorizontalBox)
+					+ SHorizontalBox::Slot()
+					.AutoWidth()
+					.VAlign(VAlign_Center)
+					.Padding(3.0f)
+					[
+						SNew(SImage)
+							.Image(FAppStyle::Get().GetBrush("Icons.Tag"))
+					]
+					+ SHorizontalBox::Slot()
+					.FillWidth(1.0)
+					.AutoWidth()
+					.Padding(3.0f)
+					[
+						SAssignNew(InlineWidget, SInlineEditableTextBlock)
+							.HighlightText_Lambda([this]() { return SearchText; })
+							.Text(FText::FromName(InItem->Label))
+							.OnVerifyTextChanged_Lambda([InItem](const FText& InText, FText& OutErrorMessage)
+								{
+									const FName OldName = InItem->Label;
+									const FName NewName = FName(InText.ToString());
+
+									if (InItem->Label == NewName)
+									{
+										return true;
+									}
+
+									check(InItem->CanRenameTo.IsBound());
+									if (InItem->CanRenameTo.Execute(OldName, NewName))
+									{
+										return true;
+									}
+									else
+									{
+										OutErrorMessage = LOCTEXT("SetNameTaken", "Name already taken");
+										return false;
+									}
+								})
+							.OnTextCommitted_Lambda([InItem](const FText& InText, ETextCommit::Type CommitInfo)
+								{
+									const FName OldName = InItem->Label;
+									const FName NewName = FName(InText.ToString());
+
+									InItem->OnRenamed.ExecuteIfBound(OldName, NewName);
+								})
+					]
+
+			];
+
+		InItem->OnRequestRename.BindSP(InlineWidget.Get(), &SInlineEditableTextBlock::EnterEditingMode);
+
+		return RowWidget;
+	}
+
+	TSharedPtr<SWidget> SLabelCollection::ListView_OnContextMenuOpening()
+	{
+		FMenuBuilder MenuBuilder(true, CommandList.ToSharedRef());
+
+		const bool bIsLabelSelected = !ListView->GetSelectedItems().IsEmpty();
+
+		MenuBuilder.AddMenuEntry(
+			LOCTEXT("AddLabel_Label", "Add Label"),
+			LOCTEXT("AddLabel_Description", "Add a new label"),
+			FSlateIcon(FAppStyle::GetAppStyleSetName(), "Icons.Plus"),
+			FUIAction(
+				FExecuteAction::CreateSP(this, &SLabelCollection::OnAddLabel)
+			));
+
+		if (bIsLabelSelected)
+		{
+			MenuBuilder.AddMenuEntry(FGenericCommands::Get().Rename);
+			MenuBuilder.AddMenuEntry(FGenericCommands::Get().Delete);
+		}
+
+		return MenuBuilder.MakeWidget();
+	}
+
+	void SLabelCollection::ListView_OnItemScrolledIntoView(FListItemPtr InItem, const TSharedPtr<ITableRow>& InWidget)
+	{
+		if (ItemToRename)
+		{
+			ItemToRename->OnRequestRename.ExecuteIfBound();
+			ItemToRename = nullptr;
+		}
+	}
+
+	void SLabelCollection::ListView_OnMouseButtonDoubleClick(FListItemPtr InItem)
+	{
+		InItem->OnRequestRename.ExecuteIfBound();
+	}
+
+	void SLabelCollection::OnAddLabel()
+	{
+		SearchText = FText();
+		
+		const FScopedTransaction Transaction(LOCTEXT("AddLabel", "Set Label"));
+		LabelCollection->Modify();
+
+		const FText DefaultLabelFormat = LOCTEXT("DefaultLabel", "NewLabel_{0}");
+
+		int32 Suffix = 0;
+		FName LabelName = FName(FText::Format(DefaultLabelFormat, Suffix).ToString());
+		while (LabelCollection->HasLabel(LabelName))
+		{
+			++Suffix;
+			LabelName = FName(FText::Format(DefaultLabelFormat, Suffix).ToString());
+		}
+
+		const bool bSuccess = LabelCollection->AddLabel(LabelName);
+		check(bSuccess);
+
+		RepopulateListData();
+
+		const FListItemPtr* const LabelItem = ListItems.FindByPredicate([LabelName](const FListItemPtr& InItem)
+		{
+			return InItem->Label == LabelName;
+		});
+		check(LabelItem);
+
+		ItemToRename = *LabelItem;
+		ListView->RequestScrollIntoView(*LabelItem);
+	}
+
+	void SLabelCollection::OnRenameLabel()
+	{
+		const TArray<FListItemPtr> SelectedItems = ListView->GetSelectedItems();
+		if (!SelectedItems.IsEmpty())
+		{
+			FListItemPtr LabelItem = SelectedItems[0];
+
+			ItemToRename = LabelItem;
+			ListView->RequestScrollIntoView(LabelItem);
+		}
+	}
+
+	void SLabelCollection::OnRemoveLabel()
+	{
+		const TArray<FListItemPtr> SelectedItems = ListView->GetSelectedItems();
+		if (!SelectedItems.IsEmpty())
+		{
+			const FScopedTransaction Transaction(LOCTEXT("RemoveLabel", "Remove Label"));
+			LabelCollection->Modify();
+
+			const FName LabelToRemove = SelectedItems[0]->Label;
+
+			const bool bSuccess = LabelCollection->RemoveLabel(LabelToRemove);
+			check(bSuccess);
+		
+			RepopulateListData();
+		}
+	}
+
+	void SLabelCollection::RepopulateListData()
+	{
+		ListItems.Reset();
+
+		for (const FName& Label : LabelCollection->GetLabels())
+		{
+			FListItemPtr LabelItem = MakeShared<FListItem>();
+			LabelItem->Label = Label;
+			LabelItem->OnRenamed = FListItem::FOnRenamed::CreateSP(this, &SLabelCollection::OnItemRenamed);
+			LabelItem->CanRenameTo = FListItem::FCanRenameTo::CreateSP(this, &SLabelCollection::CanRenameItem);
+
+			ListItems.Add(MoveTemp(LabelItem));
+		}
+
+		ApplySearchFilter();
+
+		ListView->RebuildList();
+	}
+
+	void SLabelCollection::BindCommands()
+	{
+		CommandList = MakeShareable(new FUICommandList);
+
+		CommandList->MapAction(
+			FGenericCommands::Get().Rename,
+			FExecuteAction::CreateSP(this, &SLabelCollection::OnRenameLabel)
+		);
+		CommandList->MapAction(
+			FGenericCommands::Get().Delete,
+			FExecuteAction::CreateSP(this, &SLabelCollection::OnRemoveLabel)
+		);
+	}
+
+	void SLabelCollection::OnItemRenamed(const FName OldLabel, const FName NewLabel)
+	{
+		if (NewLabel != OldLabel)
+		{
+			const bool bSuccess = LabelCollection->RenameLabel(OldLabel, NewLabel);
+			check(bSuccess);
+
+			RepopulateListData();
+		}
+	}
+
+	bool SLabelCollection::CanRenameItem(const FName OldLabel, const FName NewLabel) const
+	{
+		return !LabelCollection->HasLabel(NewLabel);
+	}
+	
+	void SLabelCollection::ApplySearchFilter()
+	{
+		const FString SearchString = SearchText.ToString();
+
+		FilteredItems = ListItems.FilterByPredicate([&](const FListItemPtr ListItem)
+		{
+			return ListItem->Label.ToString().Contains(SearchString, ESearchCase::Type::IgnoreCase);
+		});
+
+		ListView->RequestListRefresh();
+	}
+}
+
+#undef LOCTEXT_NAMESPACE
\ No newline at end of file
diff --git a/Engine/Plugins/Experimental/Animation/SkeletonTemplateFramework/Source/Editor/Private/SkeletonTemplateFramework/Labels/SLabelCollection.h b/Engine/Plugins/Experimental/Animation/SkeletonTemplateFramework/Source/Editor/Private/SkeletonTemplateFramework/Labels/SLabelCollection.h
new file mode 100644
--- /dev/null
+++ b/Engine/Plugins/Experimental/Animation/SkeletonTemplateFramework/Source/Editor/Private/SkeletonTemplateFramework/Labels/SLabelCollection.h
@@ -0,0 +1,90 @@
+// Copyright Epic Games, Inc. All Rights Reserved.
+
+#pragma once
+
+#include "CoreMinimal.h"
+#include "EditorUndoClient.h"
+#include "Widgets/SCompoundWidget.h"
+
+class FUICommandList;
+class ITableRow;
+class UAbstractSkeletonLabelCollection;
+template<typename> class SListView;
+class SSearchBox;
+class STableViewBase;
+
+namespace UE::Anim::SkelTemplates::Labels
+{
+	class SLabelCollection : public SCompoundWidget, public FSelfRegisteringEditorUndoClient
+	{
+		SLATE_BEGIN_ARGS(SLabelCollection) {}
+		SLATE_END_ARGS()
+
+	public:
+		struct FListItem
+		{
+			FName Label;
+
+			DECLARE_DELEGATE(FOnRequestRename);
+			FOnRequestRename OnRequestRename;
+
+			DECLARE_DELEGATE_RetVal_TwoParams(bool, FCanRenameTo, const FName /* Old Name */, const FName /* New Name */);
+			FCanRenameTo CanRenameTo;
+
+			DECLARE_DELEGATE_TwoParams(FOnRenamed, const FName /* Old Name */, const FName /* New Name */);
+			FOnRenamed OnRenamed;
+		};
+
+		using FListItemPtr = TSharedPtr<FListItem>;
+
+		virtual void Construct(const FArguments& InArgs, TWeakObjectPtr<UAbstractSkeletonLabelCollection> InLabelCollection);
+
+		/** FSelfRegisteringEditorUndoClient interface */
+		virtual void PostUndo(bool bSuccess) override;
+		virtual void PostRedo(bool bSuccess) override;
+
+	private:
+		virtual FReply OnKeyDown(const FGeometry& MyGeometry, const FKeyEvent& InKeyEvent) override;
+
+		TSharedRef<ITableRow> ListView_OnGenerateRow(FListItemPtr InItem, const TSharedRef<STableViewBase>& OwnerTable);
+
+		TSharedPtr<SWidget> ListView_OnContextMenuOpening();
+
+		void ListView_OnItemScrolledIntoView(FListItemPtr InItem, const TSharedPtr<ITableRow>& InWidget);
+
+		void ListView_OnMouseButtonDoubleClick(FListItemPtr InItem);
+
+		void OnAddLabel();
+
+		void OnRenameLabel();
+
+		void OnRemoveLabel();
+
+		void RepopulateListData();
+
+		void BindCommands();
+
+		void OnItemRenamed(const FName OldLabel, const FName NewLabel);
+
+		bool CanRenameItem(const FName OldLabel, const FName NewLabel) const;
+
+		void ApplySearchFilter();
+
+	private:
+		TWeakObjectPtr<UAbstractSkeletonLabelCollection> LabelCollection;
+
+		TSharedPtr<SListView<FListItemPtr>> ListView;
+
+		TArray<FListItemPtr> ListItems;
+		
+		TArray<FListItemPtr> FilteredItems;
+
+		TSharedPtr<SSearchBox> SearchBox;
+
+		TSharedPtr<FUICommandList> CommandList;
+
+		FListItemPtr ItemToRename;
+
+		FText SearchText;
+	};
+}
\ No newline at end of file
diff --git a/Engine/Plugins/Experimental/Animation/SkeletonTemplateFramework/Source/Editor/Private/SkeletonTemplateFramework/Labels/SLabelSkeletonTree.cpp b/Engine/Plugins/Experimental/Animation/SkeletonTemplateFramework/Source/Editor/Private/SkeletonTemplateFramework/Labels/SLabelSkeletonTree.cpp
new file mode 100644
--- /dev/null
+++ b/Engine/Plugins/Experimental/Animation/SkeletonTemplateFramework/Source/Editor/Private/SkeletonTemplateFramework/Labels/SLabelSkeletonTree.cpp
@@ -0,0 +1,408 @@
+// Copyright Epic Games, Inc. All Rights Reserved.
+
+#include "SkeletonTemplateFramework/Labels/SLabelSkeletonTree.h"
+
+#include "Animation/Skeleton.h"
+#include "Framework/Commands/UIAction.h"
+#include "Framework/Commands/UICommandList.h"
+#include "Framework/MultiBox/MultiBoxBuilder.h"
+#include "IPersonaToolkit.h"
+#include "IPersonaPreviewScene.h"
+#include "ReferenceSkeleton.h"
+#include "ScopedTransaction.h"
+#include "SkeletonTemplateFramework/AbstractSkeletonLabelBinding.h"
+#include "SkeletonTemplateFramework/AbstractSkeletonLabelCollection.h"
+#include "SkeletonTemplateFramework/IAbstractSkeletonEditor.h"
+#include "SkeletonTemplateFramework/Labels/ILabelBinding.h"
+#include "SkeletonTemplateFramework/Labels/ILabelsTab.h"
+#include "SkeletonTemplateFramework/Labels/BoneDragDrop.h"
+#include "SkeletonTemplateFramework/Labels/LabelDragDrop.h"
+#include "Widgets/Input/SSearchBox.h"
+#include "Widgets/Input/SHyperlink.h"
+#include "Widgets/Views/STreeView.h"
+
+#define LOCTEXT_NAMESPACE "UE::Anim::SkelTemplates::Labels::SLabelSkeletonTree"
+
+namespace UE::Anim::SkelTemplates::Labels
+{
+
+	void SLabelSkeletonTree::Construct(const FArguments& InArgs, TWeakObjectPtr<UAbstractSkeletonLabelBinding> InLabelBinding, TSharedPtr<ILabelsTab> InLabelsTabInterface)
+	{
+		LabelBinding = InLabelBinding;
+		OnTreeRefreshed = InArgs._OnTreeRefreshed;
+		LabelsTabInterface = InLabelsTabInterface;
+
+		TWeakPtr<IAbstractSkeletonEditor> AbstractSkeletonEditor = LabelsTabInterface->GetAbstractSkeletonEditor();
+		TSharedPtr<IPersonaToolkit> PersonaToolkit = AbstractSkeletonEditor.Pin()->GetPersonaToolkit();
+		TSharedPtr<IPersonaPreviewScene> PreviewScene = PersonaToolkit->GetPreviewScene();
+		PreviewScene->RegisterOnSelectedBonesChanged(FOnSelectedBonesChanged::CreateSP(this, &SLabelSkeletonTree::HandleBoneSelectionChanged));
+
+		ChildSlot
+		[
+			SNew(SVerticalBox)
+				+ SVerticalBox::Slot()
+				.AutoHeight()
+				.Padding(FMargin(1.0f, 1.0f))
+				[
+					SNew(SHorizontalBox)
+						+ SHorizontalBox::Slot()
+						.FillWidth(1.0f)
+						.Padding(FMargin(1.0f, 1.0f))
+						[
+							SAssignNew(SearchBox, SSearchBox)
+								.SelectAllTextWhenFocused(true)
+								.OnTextChanged_Lambda([](const FText& InText) { /* TODO: Implement */ })
+								.HintText(LOCTEXT("SearchBox_Hint", "Search Labels..."))
+						]
+				]
+				+ SVerticalBox::Slot()
+				.FillHeight(1.0f)
+				.Padding(FMargin(1.0f, 1.0f))
+				[
+					SAssignNew(TreeView, STreeView<FTreeItemPtr>)
+						.SelectionMode(ESelectionMode::Single)
+						.HighlightParentNodesForSelection(true)
+						.TreeItemsSource(&RootItems)
+						.OnGenerateRow(this, &SLabelSkeletonTree::TreeView_OnGenerateRow)
+						.OnGetChildren(this, &SLabelSkeletonTree::TreeView_OnGetChildren)
+						.OnContextMenuOpening(this, &SLabelSkeletonTree::TreeView_OnContextMenuOpening)
+						.OnSelectionChanged_Lambda([this](FTreeItemPtr InItem, ESelectInfo::Type InSelectType)
+						{
+							TWeakPtr<IAbstractSkeletonEditor> AbstractSkeletonEditor = LabelsTabInterface->GetAbstractSkeletonEditor();
+							TSharedPtr<IPersonaToolkit> PersonaToolkit = AbstractSkeletonEditor.Pin()->GetPersonaToolkit();
+
+							if (InItem.IsValid())
+							{
+								PersonaToolkit->GetPreviewScene()->SetSelectedBone(InItem->BoneName, InSelectType);
+							}
+							else
+							{
+								PersonaToolkit->GetPreviewScene()->SetSelectedBone(NAME_None, InSelectType);
+							}
+						})
+				]
+		];
+
+		RepopulateTreeData();
+		ExpandAllTreeItems();
+		BindCommands();
+	}
+
+	FReply SLabelSkeletonTree::OnKeyDown(const FGeometry& MyGeometry, const FKeyEvent& InKeyEvent)
+	{
+		if (CommandList->ProcessCommandBindings(InKeyEvent))
+		{
+			return FReply::Handled();
+		}
+
+		return FReply::Unhandled();
+	}
+
+	TSharedRef<ITableRow> SLabelSkeletonTree::TreeView_OnGenerateRow(FTreeItemPtr InItem, const TSharedRef<STableViewBase>& OwnerTable)
+	{
+		TSharedPtr<SHorizontalBox> HBox;
+
+		TSharedRef<ITableRow> RowWidget = SNew(STableRow<FTreeItemPtr>, OwnerTable)
+			.ShowWires(true)
+			.Style(FAppStyle::Get(), "TableView.AlternatingRow")
+			.OnDragDetected_Lambda([this, InItem](const FGeometry& MyGeometry, const FPointerEvent& MouseEvent)
+				{
+					if (MouseEvent.IsMouseButtonDown(EKeys::LeftMouseButton))
+					{
+						const TSharedRef<FBoneDragDropOp> DragDropOp = FBoneDragDropOp::New(LabelBinding, InItem->BoneName);
+						return FReply::Handled().BeginDragDrop(DragDropOp);
+					}
+
+					return FReply::Unhandled();
+				})
+			.OnCanAcceptDrop_Lambda([](const FDragDropEvent& DragDropEvent, EItemDropZone DropZone, TSharedPtr<FTreeItem> TargetTreeItem) -> TOptional<EItemDropZone>
+				{
+					TOptional<EItemDropZone> ReturnedDropZone;
+
+					const TSharedPtr<FLabelDragDropOp> DragDropOp = DragDropEvent.GetOperationAs<FLabelDragDropOp>();
+					if (DragDropOp.IsValid())
+					{
+						ReturnedDropZone = EItemDropZone::OntoItem;
+					}
+
+					return ReturnedDropZone;
+				})
+			.OnAcceptDrop_Lambda([this](const FDragDropEvent& DragDropEvent, EItemDropZone DropZone, TSharedPtr<FTreeItem> TargetTreeItem) -> FReply
+				{
+					const TSharedPtr<FLabelDragDropOp> DragDropOp = DragDropEvent.GetOperationAs<FLabelDragDropOp>();
+					
+					if (DragDropOp.IsValid())
+					{
+						const FScopedTransaction Transaction(LOCTEXT("BindBoneToLabel", "Bind Bone to Label"));
+						LabelBinding->Modify();
+
+						const bool bSuccess = LabelBinding->BindBoneToLabel(DragDropOp->LabelCollection.Get(), TargetTreeItem->BoneName, DragDropOp->Label);
+						check(bSuccess);
+
+						RepopulateTreeData();
+						LabelsTabInterface->GetLabelBindingWidget()->RepopulateTreeData();
+
+						return FReply::Handled();
+					}
+
+					return FReply::Unhandled();
+				})
+			[
+				SAssignNew(HBox, SHorizontalBox)
+					+ SHorizontalBox::Slot()
+					.AutoWidth()
+					.Padding(FMargin(3.0f))
+					[
+						SNew(SImage)
+							.ColorAndOpacity(InItem->AssignedLabels.IsEmpty() ? FSlateColor::UseForeground() : FAppStyle::Get().GetSlateColor("Colors.AccentBlue"))
+							.Image(FAppStyle::Get().GetBrush("SkeletonTree.Bone"))
+					]
+					+ SHorizontalBox::Slot()
+					.AutoWidth()
+					.Padding(FMargin(3.0f))
+					[
+						SNew(STextBlock)
+							.Text(FText::FromName(InItem->BoneName))
+					]
+			];
+
+		for (const FTreeItem::FBindings Binding : InItem->AssignedLabels)
+		{
+			HBox->AddSlot()
+				.AutoWidth()
+				.VAlign(VAlign_Center)
+				.Padding(FMargin(4.0f, 0.0f))
+				[
+					SNew(SHorizontalBox)
+					+ SHorizontalBox::Slot()
+					.AutoWidth()
+					.VAlign(VAlign_Center)
+					.Padding(FMargin(2.0f, 0.0f))
+					[
+						SNew(SImage)
+							.ColorAndOpacity(FAppStyle::Get().GetSlateColor("Colors.AccentBlue"))
+							.DesiredSizeOverride(FVector2D(16, 16))
+							.Image(FAppStyle::Get().GetBrush("Icons.Tag"))
+					]
+					+ SHorizontalBox::Slot()
+					.AutoWidth()
+					.VAlign(VAlign_Center)
+					.Padding(FMargin(2.0f, 0.0f))
+					[
+						SNew(SHyperlink)
+							.Text(FText::FromName(Binding.Label))
+							.OnNavigate_Lambda([this, Binding]()
+							{
+								TSharedPtr<ILabelBindingWidget> LabelBindingWidget = LabelsTabInterface->GetLabelBindingWidget();
+								LabelBindingWidget->ScrollToLabel(Binding.LabelCollection.Get(), Binding.Label);
+							})
+					]
+				];
+		}
+
+		return RowWidget;
+	}
+
+	void SLabelSkeletonTree::TreeView_OnGetChildren(FTreeItemPtr InItem, TArray<FTreeItemPtr>& OutChildren)
+	{
+		OutChildren.Append(InItem->Children);
+	}
+
+	TSharedPtr<SWidget> SLabelSkeletonTree::TreeView_OnContextMenuOpening()
+	{
+		FMenuBuilder MenuBuilder(true, CommandList.ToSharedRef());
+
+		TArray<FTreeItemPtr> Selection = TreeView->GetSelectedItems();
+		if (Selection.Num() > 0)
+		{
+			FTreeItemPtr Selected = Selection[0]; // Single selection only
+
+			if (Selected->AssignedLabels.Num() >= 2)
+			{
+				MenuBuilder.AddMenuEntry(
+					LOCTEXT("UnbindAll_Label", "Unbind All"),
+					LOCTEXT("UnbindAll_Tooltip", "Unbind all labels"),
+					FSlateIcon(FAppStyle::GetAppStyleSetName(), "Icons.Unlink"),
+					FUIAction(
+						FExecuteAction::CreateLambda([this, Selected]()
+							{
+								const FScopedTransaction Transaction(LOCTEXT("UnbindAll", "Unbind All"));
+								LabelBinding->Modify();
+
+								for (const FTreeItem::FBindings Binding : Selected->AssignedLabels)
+								{
+									const bool bSuccess = LabelBinding->UnbindBoneFromLabel(Binding.LabelCollection.Get(), Selected->BoneName, Binding.Label);
+									check(bSuccess);
+								}
+
+								RepopulateTreeData();
+								LabelsTabInterface->GetLabelBindingWidget()->RepopulateTreeData();
+							}
+						)
+					));
+
+				MenuBuilder.AddSeparator();
+			}
+
+			for (const FTreeItem::FBindings Binding : Selected->AssignedLabels)
+			{
+				MenuBuilder.AddMenuEntry(
+					FText::Format(LOCTEXT("UnbindBone_Label", "Unbind {0}"), FText::FromName(Binding.Label)),
+					LOCTEXT("UnbindBone_Tooltip", "Unbind the selected bone and label"),
+					FSlateIcon(FAppStyle::GetAppStyleSetName(), "Icons.Unlink"),
+					FUIAction(
+						FExecuteAction::CreateLambda([this, Binding, Selected]()
+							{
+								const FScopedTransaction Transaction(FText::Format(LOCTEXT("UnbindBone_Label", "Unbind {0}"), FText::FromName(Binding.Label)));
+								LabelBinding->Modify();
+
+								const bool bSuccess = LabelBinding->UnbindBoneFromLabel(Binding.LabelCollection.Get(), Selected->BoneName, Binding.Label);
+								check(bSuccess);
+
+								RepopulateTreeData();
+								LabelsTabInterface->GetLabelBindingWidget()->RepopulateTreeData();
+							}
+						)
+					));
+			}
+
+			return MenuBuilder.MakeWidget();
+		}
+
+		return SNullWidget::NullWidget;
+	}
+
+	void SLabelSkeletonTree::PostUndo(bool bSuccess)
+	{
+		RepopulateTreeData();
+	}
+
+	void SLabelSkeletonTree::PostRedo(bool bSuccess)
+	{
+		RepopulateTreeData();
+	}
+
+	void SLabelSkeletonTree::ScrollToBone(const FName InBoneName)
+	{
+		for (FTreeItemPtr TreeItem : GetAllTreeItems())
+		{
+			if (TreeItem->BoneName == InBoneName)
+			{
+				TreeView->RequestScrollIntoView(TreeItem);
+				TreeView->SetSelection(TreeItem);
+				return;
+			}
+		}
+	}
+
+	void SLabelSkeletonTree::RepopulateTreeData()
+	{
+		if (!bRepopulating)
+		{
+			TGuardValue<bool> RecursionGuard(bRepopulating, true);
+			
+			// Rebuild tree data
+			RootItems.Reset();
+
+			if (LabelBinding.IsValid())
+			{
+				TMap<int32, FTreeItemPtr> TreeItemsMap;
+
+				if (const TObjectPtr<const USkeleton> Skeleton = LabelBinding->GetSkeleton())
+				{
+					const FReferenceSkeleton& RefSkeleton = Skeleton->GetReferenceSkeleton();
+					const int32 BoneCount = RefSkeleton.GetNum();
+
+					for (int32 BoneIndex = 0; BoneIndex < BoneCount; ++BoneIndex)
+					{
+						TSharedPtr<FTreeItem> TreeItem = MakeShared<FTreeItem>();
+						TreeItem->BoneName = RefSkeleton.GetBoneName(BoneIndex);
+
+						const int32 ParentIndex = RefSkeleton.GetParentIndex(BoneIndex);
+
+						if (ParentIndex == INDEX_NONE)
+						{
+							RootItems.Add(TreeItem);
+						}
+						else
+						{
+							if (ensure(TreeItemsMap.Contains(ParentIndex)))
+							{
+								TreeItemsMap[ParentIndex]->Children.Add(TreeItem);
+							}
+						}
+
+						TreeItemsMap.Add(BoneIndex, TreeItem);
+					}
+
+					for (const FAbstractSkeleton_LabelBinding& Binding : LabelBinding->GetLabelBindings())
+					{
+						const int32 BoneIndex = RefSkeleton.FindBoneIndex(Binding.BoneName);
+						if (ensure(BoneIndex != INDEX_NONE && TreeItemsMap.Contains(BoneIndex)))
+						{
+							FTreeItem::FBindings& Bindings = TreeItemsMap[BoneIndex]->AssignedLabels.AddDefaulted_GetRef();
+							Bindings.Label = Binding.Label;
+							Bindings.LabelCollection = Binding.SourceCollection;
+						}
+					}
+				}
+			}
+
+			TreeView->RequestTreeRefresh();
+			
+			for (FTreeItemPtr TreeItem : GetAllTreeItems())
+			{
+				TreeView->SetItemExpansion(TreeItem, true);
+			}
+
+			OnTreeRefreshed.ExecuteIfBound();
+		}
+	}
+
+	void SLabelSkeletonTree::ExpandAllTreeItems()
+	{
+		for (FTreeItemPtr TreeItem : GetAllTreeItems())
+		{
+			TreeView->SetItemExpansion(TreeItem, true);
+		}
+	}
+
+	void SLabelSkeletonTree::BindCommands()
+	{
+		CommandList = MakeShareable(new FUICommandList);
+	}
+
+	TArray<TSharedPtr<SLabelSkeletonTree::FTreeItem>> SLabelSkeletonTree::GetAllTreeItems()
+	{
+		TArray<FTreeItemPtr> AllItems;
+		AllItems.Append(RootItems);
+
+		for (int32 Index = 0; Index < AllItems.Num(); ++Index)
+		{
+			AllItems.Append(AllItems[Index]->Children);
+		}
+
+		return AllItems;
+	}
+
+	void SLabelSkeletonTree::HandleBoneSelectionChanged(const TArray<FName>& InBoneNames, ESelectInfo::Type InSelectInfo)
+	{
+		const FName SelectedBoneName = InBoneNames.IsEmpty() ? NAME_None : InBoneNames[0]; // Single selection only
+
+		for (const FTreeItemPtr& TreeItem : GetAllTreeItems())
+		{
+			if (TreeItem->BoneName == SelectedBoneName)
+			{
+				if (!TreeView->IsItemSelected(TreeItem)) // Recursion guard
+				{
+					TreeView->RequestScrollIntoView(TreeItem);
+					TreeView->SetSelection(TreeItem);
+					return;
+				}
+
+			}
+		}
+	}
+}
+
+#undef LOCTEXT_NAMESPACE
\ No newline at end of file
diff --git a/Engine/Plugins/Experimental/Animation/SkeletonTemplateFramework/Source/Editor/Private/SkeletonTemplateFramework/Labels/SLabelSkeletonTree.h b/Engine/Plugins/Experimental/Animation/SkeletonTemplateFramework/Source/Editor/Private/SkeletonTemplateFramework/Labels/SLabelSkeletonTree.h
new file mode 100644
--- /dev/null
+++ b/Engine/Plugins/Experimental/Animation/SkeletonTemplateFramework/Source/Editor/Private/SkeletonTemplateFramework/Labels/SLabelSkeletonTree.h
@@ -0,0 +1,90 @@
+// Copyright Epic Games, Inc. All Rights Reserved.
+
+#pragma once
+
+#include "CoreMinimal.h"
+#include "EditorUndoClient.h"
+#include "SkeletonTemplateFramework/Labels/ILabelSkeletonTree.h"
+#include "Widgets/SCompoundWidget.h"
+
+class FUICommandList;
+class ITableRow;
+class STableViewBase;
+class SSearchBox;
+template<typename> class STreeView;
+class UAbstractSkeletonLabelBinding;
+class UAbstractSkeletonLabelCollection;
+
+namespace UE::Anim::SkelTemplates::Labels
+{
+	class ILabelsTab;
+
+	class SLabelSkeletonTree : public SCompoundWidget, public FSelfRegisteringEditorUndoClient, public ILabelSkeletonTreeWidget
+	{
+		SLATE_BEGIN_ARGS(SLabelSkeletonTree) {}
+			SLATE_EVENT(FSimpleDelegate, OnTreeRefreshed)
+		SLATE_END_ARGS()
+
+	public:
+		struct FTreeItem
+		{
+			struct FBindings
+			{
+				TWeakObjectPtr<const UAbstractSkeletonLabelCollection> LabelCollection;
+				FName Label;
+			};
+
+			FName BoneName;
+			TArray<TSharedPtr<FTreeItem>> Children;
+			TArray<FBindings> AssignedLabels;
+		};
+
+		using FTreeItemPtr = TSharedPtr<FTreeItem>;
+
+		void Construct(const FArguments& InArgs, TWeakObjectPtr<UAbstractSkeletonLabelBinding> InSetBinding, TSharedPtr<ILabelsTab> InLabelsTabInterface);
+
+		/** FSelfRegisteringEditorUndoClient interface */
+		virtual void PostUndo(bool bSuccess) override;
+		virtual void PostRedo(bool bSuccess) override;
+
+		/** ILabelSkeletonTreeWidget */
+		virtual void ScrollToBone(const FName InBoneName) override;
+		virtual void RepopulateTreeData() override;
+
+	private:
+		virtual FReply OnKeyDown(const FGeometry& MyGeometry, const FKeyEvent& InKeyEvent) override;
+
+		TSharedRef<ITableRow> TreeView_OnGenerateRow(FTreeItemPtr InItem, const TSharedRef<STableViewBase>& OwnerTable);
+
+		void TreeView_OnGetChildren(FTreeItemPtr InItem, TArray<FTreeItemPtr>& OutChildren);
+
+		TSharedPtr<SWidget> TreeView_OnContextMenuOpening();
+
+		void ExpandAllTreeItems();
+
+		void BindCommands();
+
+		TArray<FTreeItemPtr> GetAllTreeItems();
+
+		// Triggered when a bone is selected in the viewport
+		void HandleBoneSelectionChanged(const TArray<FName>& InBoneNames, ESelectInfo::Type InSelectInfo);
+
+	public:
+		TWeakObjectPtr<UAbstractSkeletonLabelBinding> LabelBinding;
+
+	private:
+		TSharedPtr<STreeView<FTreeItemPtr>> TreeView;
+
+		TArray<FTreeItemPtr> RootItems;
+
+		TSharedPtr<SSearchBox> SearchBox;
+
+		TSharedPtr<FUICommandList> CommandList;
+
+		FSimpleDelegate OnTreeRefreshed;
+
+		bool bRepopulating = false;
+
+		TSharedPtr<ILabelsTab> LabelsTabInterface;
+	};
+}
\ No newline at end of file
diff --git a/Engine/Plugins/Experimental/Animation/SkeletonTemplateFramework/Source/Editor/Private/SkeletonTemplateFramework/Labels/SLabelsTab.cpp b/Engine/Plugins/Experimental/Animation/SkeletonTemplateFramework/Source/Editor/Private/SkeletonTemplateFramework/Labels/SLabelsTab.cpp
new file mode 100644
--- /dev/null
+++ b/Engine/Plugins/Experimental/Animation/SkeletonTemplateFramework/Source/Editor/Private/SkeletonTemplateFramework/Labels/SLabelsTab.cpp
@@ -0,0 +1,220 @@
+// Copyright Epic Games, Inc. All Rights Reserved.
+
+#include "SkeletonTemplateFramework/Labels/SLabelsTab.h"
+
+#include "Animation/Skeleton.h"
+#include "AssetThumbnail.h"
+#include "Framework/Docking/LayoutService.h"
+#include "Framework/MultiBox/MultiBoxBuilder.h"
+#include "SkeletonTemplateFramework/AbstractSkeletonLabelBinding.h"
+#include "SkeletonTemplateFramework/Labels/SLabelBinding.h"
+#include "SkeletonTemplateFramework/Labels/SLabelSkeletonTree.h"
+#include "PropertyCustomizationHelpers.h"
+#include "Widgets/Docking/SDockTab.h"
+
+#define LOCTEXT_NAMESPACE "UE::Anim::SkelTemplates::Labels::SLabelsTab"
+
+namespace UE::Anim::SkelTemplates::Labels
+{
+	FName SLabelsTab::FTabs::SkeletonTreeId("AbstractSkeletonBindingEditor_Labels_SkeletonTreeTab");
+	FName SLabelsTab::FTabs::LabelBindingsId("AbstractSkeletonBindingEditor_Labels_BindingsTab");
+
+	TSharedPtr<ILabelBindingWidget> SLabelsTab::GetLabelBindingWidget() const
+	{
+		return LabelBindingWidget;
+	}
+
+	TSharedPtr<ILabelSkeletonTreeWidget> SLabelsTab::GetLabelSkeletonTreeWidget() const
+	{
+		return LabelSkeletonTreeWidget;
+	}
+
+	void SLabelsTab::RepopulateLabelData()
+	{
+		GetLabelBindingWidget()->RepopulateTreeData();
+		GetLabelSkeletonTreeWidget()->RepopulateTreeData();
+	}
+
+	void SLabelsTab::Construct(const FArguments& InArgs)
+	{
+		LabelBinding = InArgs._LabelBinding;
+		AbstractSkeletonEditor = InArgs._AbstractSkeletonEditor;
+
+		static const FName LayoutName = "UE::Anim::STF::LabelsTab::SLabelsTab_Layout";
+
+		TSharedRef<FTabManager::FLayout> Layout = FTabManager::NewLayout(LayoutName)
+			->AddArea
+			(
+				FTabManager::NewPrimaryArea()
+				->SetOrientation(Orient_Horizontal)
+				->Split
+				(
+					FTabManager::NewStack()
+					->AddTab(FTabs::SkeletonTreeId, ETabState::OpenedTab)
+				)
+				->Split
+				(
+					FTabManager::NewStack()
+					->AddTab(FTabs::LabelBindingsId, ETabState::OpenedTab)
+				)
+			);
+
+		check(InArgs._ParentTab.IsValid());
+		TabManager = FGlobalTabmanager::Get()->NewTabManager(InArgs._ParentTab.ToSharedRef());
+		TabManager->SetOnPersistLayout(FTabManager::FOnPersistLayout::CreateSP(this, &SLabelsTab::HandleTabManagerPersistLayout));
+
+		TabManager->RegisterTabSpawner(FTabs::SkeletonTreeId, FOnSpawnTab::CreateLambda([this](const FSpawnTabArgs& Args)
+			{
+				return SAssignNew(LabelSkeletonTreeTab, SDockTab)
+				[
+					SAssignNew(LabelSkeletonTreeWidget, SLabelSkeletonTree, LabelBinding, SharedThis(this))
+				];
+			}))
+			.SetDisplayName(LOCTEXT("SkeletonTree_TabDisplayName", "Skeleton Tree"))
+			.SetIcon(FSlateIcon(FAppStyle::GetAppStyleSetName(), "Persona.Tabs.SkeletonTree"));
+
+		TabManager->RegisterTabSpawner(FTabs::LabelBindingsId, FOnSpawnTab::CreateLambda([this](const FSpawnTabArgs& Args)
+			{
+				return SAssignNew(LabelBindingTab, SDockTab)
+				[
+					SAssignNew(LabelBindingWidget, SLabelBinding, LabelBinding, SharedThis(this))
+				];
+			}))
+			.SetDisplayName(LOCTEXT("LabelBindings_TabDisplayName", "Label Bindings"))
+			.SetIcon(FSlateIcon(FAppStyle::GetAppStyleSetName(), "LevelEditor.Tabs.Outliner"));
+
+		Layout = FLayoutSaveRestore::LoadFromConfig(GEditorLayoutIni, Layout);
+
+		FSlimHorizontalToolBarBuilder ToolbarBuilder(TSharedPtr<const FUICommandList>(), FMultiBoxCustomization::None);
+
+		ToolbarBuilder.BeginSection("Common");
+		{
+			ToolbarBuilder.AddWidget(
+				SNew(SHorizontalBox)
+				+ SHorizontalBox::Slot()
+				.AutoWidth()
+				.Padding(FMargin(4.0f, 0.0f, 24.0f, 0.0f))
+				.VAlign(VAlign_Center)
+				[
+					SNew(STextBlock)
+						.Text(LOCTEXT("LabelBinding", "Label Binding"))
+				]
+				+ SHorizontalBox::Slot()
+				.AutoWidth()
+				[
+					SNew(SObjectPropertyEntryBox)
+						.AllowedClass(UAbstractSkeletonLabelBinding::StaticClass())
+						.ObjectPath_Lambda([&]()
+							{
+								return LabelBinding->GetPathName();
+							})
+						.OnObjectChanged_Lambda([this](const FAssetData& AssetData)
+							{
+								LabelBinding = CastChecked<UAbstractSkeletonLabelBinding>(AssetData.GetAsset());
+
+								LabelSkeletonTreeTab->SetContent(
+									SAssignNew(LabelSkeletonTreeWidget, SLabelSkeletonTree, LabelBinding, SharedThis(this)));
+
+								LabelBindingTab->SetContent(
+									SAssignNew(LabelBindingWidget, SLabelBinding, LabelBinding, SharedThis(this)));
+							})
+						.AllowClear(false)
+						.DisplayUseSelected(true)
+						.DisplayBrowse(true)
+						.DisplayThumbnail(true)
+						.ThumbnailPool(MakeShared<FAssetThumbnailPool>(1))
+				],
+				NAME_None,
+				true,
+				HAlign_Left
+			);
+
+			ToolbarBuilder.AddSeparator();
+
+			ToolbarBuilder.AddWidget(
+				SNew(SHorizontalBox)
+				+ SHorizontalBox::Slot()
+				.AutoWidth()
+				.Padding(FMargin(16.0f, 0.0f, 24.0f, 0.0f))
+				.VAlign(VAlign_Center)
+				[
+					SNew(STextBlock)
+						.Text(LOCTEXT("Skeleton", "Skeleton"))
+				]
+				+ SHorizontalBox::Slot()
+				.AutoWidth()
+				[
+					SNew(SObjectPropertyEntryBox)
+						.AllowedClass(USkeleton::StaticClass())
+						.OnIsEnabled_Lambda([&]()
+							{
+								return LabelBinding.IsValid();
+							})
+						.ObjectPath_Lambda([&]()
+							{
+								if (LabelBinding.IsValid() && LabelBinding->GetSkeleton())
+								{
+									return LabelBinding->GetSkeleton()->GetPathName();
+								}
+								return FString();
+							})
+						.OnObjectChanged_Lambda([this](const FAssetData& AssetData)
+							{
+								if (LabelBinding.IsValid())
+								{
+									if (TObjectPtr<USkeleton> NewSkeleton = Cast<USkeleton>(AssetData.GetAsset()))
+									{
+										LabelBinding->SetSkeleton(NewSkeleton);
+
+										LabelSkeletonTreeWidget->RepopulateTreeData();
+										LabelBindingWidget->RepopulateTreeData();
+									}
+								}
+							})
+						.AllowClear(false)
+						.DisplayUseSelected(true)
+						.DisplayBrowse(true)
+						.DisplayThumbnail(true)
+						.ThumbnailPool(MakeShared<FAssetThumbnailPool>(1))
+				],
+				NAME_None,
+				true,
+				HAlign_Left
+			);
+		}
+		ToolbarBuilder.EndSection();
+
+		ToolbarBuilder.SetStyle(&FAppStyle::Get(), "SlimToolBar");
+
+		ChildSlot
+		[
+			SNew(SVerticalBox)
+				+ SVerticalBox::Slot()
+				.AutoHeight()
+				[
+					SNew(SHorizontalBox)
+						+ SHorizontalBox::Slot()
+						[
+							ToolbarBuilder.MakeWidget()
+						]
+				]
+				+ SVerticalBox::Slot()
+				.FillHeight(1.0f)
+				[
+					TabManager->RestoreFrom(Layout, nullptr).ToSharedRef()
+				]
+		];
+	}
+
+	void SLabelsTab::HandleTabManagerPersistLayout(const TSharedRef<FTabManager::FLayout>& LayoutToSave)
+	{
+		FLayoutSaveRestore::SaveToConfig(GEditorLayoutIni, LayoutToSave);
+	}
+
+	TWeakPtr<IAbstractSkeletonEditor> SLabelsTab::GetAbstractSkeletonEditor() const
+	{
+		return AbstractSkeletonEditor;
+	}
+}
+
+#undef LOCTEXT_NAMESPACE
\ No newline at end of file
diff --git a/Engine/Plugins/Experimental/Animation/SkeletonTemplateFramework/Source/Editor/Private/SkeletonTemplateFramework/Labels/SLabelsTab.h b/Engine/Plugins/Experimental/Animation/SkeletonTemplateFramework/Source/Editor/Private/SkeletonTemplateFramework/Labels/SLabelsTab.h
new file mode 100644
--- /dev/null
+++ b/Engine/Plugins/Experimental/Animation/SkeletonTemplateFramework/Source/Editor/Private/SkeletonTemplateFramework/Labels/SLabelsTab.h
@@ -0,0 +1,64 @@
+// Copyright Epic Games, Inc. All Rights Reserved.
+
+#pragma once
+
+#include "CoreMinimal.h"
+#include "Framework/Docking/TabManager.h"
+#include "Widgets/SCompoundWidget.h"
+#include "SkeletonTemplateFramework/Labels/ILabelsTab.h"
+
+class UAbstractSkeletonLabelBinding;
+class SDockTab;
+
+namespace UE::Anim::SkelTemplates
+{
+	class IAbstractSkeletonEditor;
+}
+
+namespace UE::Anim::SkelTemplates::Labels
+{
+	class SLabelBinding;
+	class SLabelSkeletonTree;
+
+	class SLabelsTab : public SCompoundWidget, public ILabelsTab
+	{
+	public:
+		struct FTabs
+		{
+			static FName SkeletonTreeId;
+			static FName LabelBindingsId;
+		};
+
+		SLATE_BEGIN_ARGS(SLabelsTab) {}
+			SLATE_ARGUMENT(TSharedPtr<SDockTab>, ParentTab)
+			SLATE_ARGUMENT(TWeakObjectPtr<UAbstractSkeletonLabelBinding>, LabelBinding)
+			SLATE_ARGUMENT(TWeakPtr<IAbstractSkeletonEditor>, AbstractSkeletonEditor)
+		SLATE_END_ARGS()
+
+		/** ILabelsTab */
+		virtual TSharedPtr<ILabelBindingWidget> GetLabelBindingWidget() const override;
+		virtual TSharedPtr<ILabelSkeletonTreeWidget> GetLabelSkeletonTreeWidget() const override;
+		virtual void RepopulateLabelData() override;
+		virtual TWeakPtr<IAbstractSkeletonEditor> GetAbstractSkeletonEditor() const override;
+
+		virtual void Construct(const FArguments& InArgs);
+
+	private:
+		void HandleTabManagerPersistLayout(const TSharedRef<FTabManager::FLayout>& LayoutToSave);
+
+	private:
+		TWeakObjectPtr<UAbstractSkeletonLabelBinding> LabelBinding;
+
+		TSharedPtr<FTabManager> TabManager;
+
+		TSharedPtr<SDockTab> LabelSkeletonTreeTab;
+
+		TSharedPtr<SLabelSkeletonTree> LabelSkeletonTreeWidget;
+		
+		TSharedPtr<SDockTab> LabelBindingTab;
+
+		TSharedPtr<SLabelBinding> LabelBindingWidget;
+
+		TWeakPtr<IAbstractSkeletonEditor> AbstractSkeletonEditor;
+	};
+}
\ No newline at end of file
diff --git a/Engine/Plugins/Experimental/Animation/SkeletonTemplateFramework/Source/Editor/Private/SkeletonTemplateFramework/Module.cpp b/Engine/Plugins/Experimental/Animation/SkeletonTemplateFramework/Source/Editor/Private/SkeletonTemplateFramework/Module.cpp
--- a/Engine/Plugins/Experimental/Animation/SkeletonTemplateFramework/Source/Editor/Private/SkeletonTemplateFramework/Module.cpp
+++ b/Engine/Plugins/Experimental/Animation/SkeletonTemplateFramework/Source/Editor/Private/SkeletonTemplateFramework/Module.cpp
@@ -2,7 +2,7 @@
 
 #include "Modules/ModuleManager.h"
 
-namespace UE::Anim::STF
+namespace UE::Anim::SkelTemplates
 {
 	class FEditorModule : public IModuleInterface
 	{
@@ -17,4 +17,4 @@ namespace UE::Anim::STF
 	};
 }
 
-IMPLEMENT_MODULE(UE::Anim::STF::FEditorModule, SkeletonTemplateFrameworkEditor)
+IMPLEMENT_MODULE(UE::Anim::SkelTemplates::FEditorModule, SkeletonTemplateFrameworkEditor)
diff --git a/Engine/Plugins/Experimental/Animation/SkeletonTemplateFramework/Source/Editor/Private/SkeletonTemplateFramework/Sets/SAttributesList.cpp b/Engine/Plugins/Experimental/Animation/SkeletonTemplateFramework/Source/Editor/Private/SkeletonTemplateFramework/Sets/SAttributesList.cpp
new file mode 100644
--- /dev/null
+++ b/Engine/Plugins/Experimental/Animation/SkeletonTemplateFramework/Source/Editor/Private/SkeletonTemplateFramework/Sets/SAttributesList.cpp
@@ -0,0 +1,331 @@
+// Copyright Epic Games, Inc. All Rights Reserved.
+
+#include "SAttributesList.h"
+
+#include "Animation/BuiltInAttributeTypes.h"
+#include "Framework/Application/SlateApplication.h"
+#include "Framework/MultiBox/MultiBoxBuilder.h"
+#include "PersonaModule.h"
+#include "SPositiveActionButton.h"
+#include "ScopedTransaction.h"
+#include "SetDragDrop.h"
+#include "SkeletonTemplateFramework/AbstractSkeletonSetBinding.h"
+#include "ToolMenus.h"
+#include "Widgets/Input/SSearchBox.h"
+
+#define LOCTEXT_NAMESPACE "UE::Anim::SkelTemplates::Sets::SAttributesList"
+
+using namespace UE::Anim::SkelTemplates::Sets;
+
+FName SAttributesList::FColumns::SetId = FName("Set");
+FName SAttributesList::FColumns::AttributeId = FName("Attribute");
+
+FName SAttributesList::FMenus::AddAttributeId = FName("SAttributesList.AddAttribute");
+
+void SAttributesList::SetSetBinding(TWeakObjectPtr<UAbstractSkeletonSetBinding> InSetBinding)
+{
+	SetBinding = InSetBinding;
+	RepopulateListData();
+}
+
+void SAttributesList::Construct(const FArguments& InArgs)
+{
+	SetBinding = InArgs._SetBinding;
+	OnListRefreshed = InArgs._OnListRefreshed;
+
+	RegisterMenus();
+
+	ChildSlot
+	[
+		SNew(SVerticalBox)
+			+ SVerticalBox::Slot()
+			.AutoHeight()
+			.Padding(FMargin(1.0f, 1.0f))
+			[
+				SNew(SHorizontalBox)
+					+ SHorizontalBox::Slot()
+					.AutoWidth()
+					.Padding(FMargin(1.0f, 1.0f))
+					[
+						SNew(SPositiveActionButton)
+							.Icon(FAppStyle::Get().GetBrush("Icons.Plus"))
+							.Text(LOCTEXT("AddAttributeButton_Label", "Add"))
+							.OnGetMenuContent( this, &SAttributesList::CreateAddAttributeWidget)
+					]
+					+ SHorizontalBox::Slot()
+					.FillWidth(1.0f)
+					.Padding(FMargin(1.0f, 1.0f))
+					[
+						SNew(SSearchBox)
+							.SelectAllTextWhenFocused(true)
+							.OnTextChanged_Lambda([](const FText& InText) { /* TODO: Implement */ })
+							.HintText(LOCTEXT("SearchBox_Hint", "Search Sets..."))
+					]
+			]
+			+ SVerticalBox::Slot()
+			[
+				SAssignNew(ListView, SListView<FListItemPtr>)
+					.ListItemsSource(&ListItems)
+					.OnGenerateRow(this, &SAttributesList::ListView_OnGenerateRow)
+					.HeaderRow(
+						SNew(SHeaderRow)
+						+ SHeaderRow::Column(FColumns::SetId)
+						.DefaultLabel(LOCTEXT("SetColumnLabel", "Set"))
+						.FillWidth(0.3f)
+						+ SHeaderRow::Column(FColumns::AttributeId)
+						.DefaultLabel(LOCTEXT("AttributeColumnLabel", "Attribute"))
+						.FillWidth(0.7f)
+					)
+			]
+	];
+
+	RepopulateListData();
+}
+
+void SAttributesList::RepopulateListData()
+{
+	if (!bRepopulating)
+	{
+		TGuardValue<bool> RecursionGuard(bRepopulating, true);
+
+		ListItems.Empty();
+
+		if (SetBinding.IsValid())
+		{
+			for (const FAbstractSkeleton_AttributeBinding& Binding : SetBinding->GetAttributeBindings())
+			{
+				FListItemPtr BindingItem = MakeShared<FListItem>(Binding.SetName, Binding.Attribute);
+				ListItems.Add(BindingItem);
+			}
+		}
+
+		ListView->RequestListRefresh();
+
+		OnListRefreshed.ExecuteIfBound();
+	}
+}
+
+class SAttributesListRow : public SMultiColumnTableRow<SAttributesList::FListItemPtr>
+{
+public:
+	SLATE_BEGIN_ARGS(SAttributesListRow) {}
+		SLATE_ARGUMENT(SAttributesList::FListItemPtr, Item)
+		SLATE_EVENT(FOnCanAcceptDrop, OnCanAcceptDrop)
+		SLATE_EVENT(FOnAcceptDrop, OnAcceptDrop)
+	SLATE_END_ARGS()
+
+	void Construct(const FArguments& InArgs, const TSharedRef<STableViewBase>& OwnerTableView)
+	{
+		Item = InArgs._Item;
+
+		const FSuperRowType::FArguments Args = FSuperRowType::FArguments()
+			.OnCanAcceptDrop(InArgs._OnCanAcceptDrop)
+			.OnAcceptDrop(InArgs._OnAcceptDrop)
+			.Style(FAppStyle::Get(), "TableView.AlternatingRow");
+
+		SMultiColumnTableRow<SAttributesList::FListItemPtr>::Construct(Args, OwnerTableView);
+	}
+
+	virtual TSharedRef<SWidget> GenerateWidgetForColumn(const FName& ColumnName) override
+	{
+		if (ColumnName == SAttributesList::FColumns::SetId)
+		{
+			return SNew(SHorizontalBox)
+				+ SHorizontalBox::Slot()
+				.AutoWidth()
+				.VAlign(VAlign_Center)
+				.Padding(6, 0)
+				[
+					SNew(SBox)
+						.WidthOverride(8.0f)
+						.HeightOverride(8.0f)
+						[
+							SNew(SImage)
+								.Image(FAppStyle::Get().GetBrush("Icons.FilledCircle"))
+								.ColorAndOpacity((Item->SetName != NAME_None) ? FLinearColor(FColor(31, 228, 75)) : FLinearColor(FColor(239, 53, 53)))
+						]
+				]
+				+ SHorizontalBox::Slot()
+				.AutoWidth()
+				.VAlign(VAlign_Center)
+				.Padding(2, 2)
+				[
+					SNew(SImage)
+						.Image(FAppStyle::Get().GetBrush("ClassIcon.GroupActor"))
+				]
+				+ SHorizontalBox::Slot()
+				.AutoWidth()
+				.VAlign(VAlign_Center)
+				.Padding(2, 2)
+				[
+					SNew(STextBlock)
+						.Text(FText::FromName(Item->SetName))
+				];
+		}
+		else if (ColumnName == SAttributesList::FColumns::AttributeId)
+		{
+			return SNew(SHorizontalBox)
+				+ SHorizontalBox::Slot()
+				.AutoWidth()
+				.VAlign(VAlign_Center)
+				.Padding(2, 2)
+				[
+					SNew(SImage)
+						.Image(FAppStyle::Get().GetBrush("AnimGraph.Attribute.Attributes.Icon"))
+						.ColorAndOpacity((Item->SetName != NAME_None) ? FLinearColor(FColor(31, 228, 75)) : FSlateColor::UseForeground())
+				]
+				+ SHorizontalBox::Slot()
+				.AutoWidth()
+				.VAlign(VAlign_Center)
+				.Padding(2, 2)
+				[
+					SNew(STextBlock)
+						.Text(FText::FromName(Item->Attribute.GetName()))
+				];
+		}
+		else 
+		{
+			return SNullWidget::NullWidget;
+		}
+	}
+
+private:
+	SAttributesList::FListItemPtr Item;
+};
+
+SAttributesList::FListItem::FListItem(const FName InSetName, const FAnimationAttributeIdentifier& InAttribute)
+	: Attribute(InAttribute)
+	, SetName(InSetName)
+{
+}
+
+TSharedRef<ITableRow> SAttributesList::ListView_OnGenerateRow(FListItemPtr Item, const TSharedRef<STableViewBase>& OwnerTable)
+{
+	return SNew(SAttributesListRow, OwnerTable)
+		.Item(Item)
+		.OnCanAcceptDrop_Lambda([](const FDragDropEvent& DragDropEvent, EItemDropZone DropZone, FListItemPtr TargetTreeItem) -> TOptional<EItemDropZone>
+			{
+				TOptional<EItemDropZone> ReturnedDropZone;
+
+				const TSharedPtr<FSetDragDropOp> DragDropOp = DragDropEvent.GetOperationAs<FSetDragDropOp>();
+				if (DragDropOp.IsValid())
+				{
+					ReturnedDropZone = EItemDropZone::OntoItem;
+				}
+
+				return ReturnedDropZone;
+			})
+		.OnAcceptDrop_Lambda([this](const FDragDropEvent& DragDropEvent, EItemDropZone DropZone, FListItemPtr TargetTreeItem) -> FReply
+			{
+				const TSharedPtr<FSetDragDropOp> DragDropOp = DragDropEvent.GetOperationAs<FSetDragDropOp>();
+
+				if (DragDropOp.IsValid())
+				{
+					const FScopedTransaction Transaction(LOCTEXT("AddAttributeToSet", "Add Attribute to Set"));
+					SetBinding->Modify();
+
+					SetBinding->AddAttributeToSet(TargetTreeItem->Attribute, DragDropOp->Item->Name);
+					
+					RepopulateListData();
+
+					return FReply::Handled();
+				}
+
+				return FReply::Unhandled();
+			});
+}
+
+void SAttributesList::RegisterMenus()
+{
+	if (UToolMenus::Get()->IsMenuRegistered(FMenus::AddAttributeId))
+	{
+		return;
+	}
+
+	FToolMenuOwnerScoped OwnerScoped(this);
+	UToolMenu* Menu = UToolMenus::Get()->RegisterMenu(FMenus::AddAttributeId);
+
+	UToolMenu* AddCurveMenu = Menu->AddSubMenu(FToolMenuOwner(), "AddCurve", "AddCurve", LOCTEXT("AddCurve", "Add Curve"));
+	UToolMenu* AddAttributeMenu = Menu->AddSubMenu(FToolMenuOwner(), "AddAttribute", "AddAttribute", LOCTEXT("AddAttribute", "Add Attribute"));
+
+	{
+		FToolMenuSection& AttributeWidgetSection = AddAttributeMenu->AddSection("WidgetSection");
+
+		AttributeWidgetSection.AddDynamicEntry(
+			"AttributePicker",
+			FNewToolMenuSectionDelegate::CreateLambda([this](FToolMenuSection& Section)
+			{
+				FPersonaModule& PersonaModule = FModuleManager::LoadModuleChecked<FPersonaModule>("Persona");
+						
+				FToolMenuEntry WidgetEntry = FToolMenuEntry::InitWidget(
+					"AttributePicker",
+					PersonaModule.CreateAttributePicker(
+						SetBinding->GetSkeleton(),
+						FOnAttributesPicked::CreateLambda([this](const TConstArrayView<FAnimationAttributeIdentifier> InSelectedAttributes)
+							{
+								const FScopedTransaction Transaction(LOCTEXT("AddAttributeToSet", "Add Attribute to Set"));
+								SetBinding->Modify();
+
+								for (const FAnimationAttributeIdentifier& Attribute : InSelectedAttributes)
+								{
+									SetBinding->AddAttributeToSet(Attribute, NAME_None);
+								}
+
+								FSlateApplication::Get().DismissAllMenus();
+								RepopulateListData();
+							}),
+						true,
+						false),
+					FText(),
+					true,
+					false,
+					true
+				);
+
+				Section.AddEntry(WidgetEntry);
+			}));
+	}
+
+	{
+		FToolMenuSection& CurveWidgetSection = AddCurveMenu->AddSection("WidgetSection");
+
+		CurveWidgetSection.AddDynamicEntry(
+			"CurvePicker",
+			FNewToolMenuSectionDelegate::CreateLambda([this](FToolMenuSection& Section)
+			{
+				FPersonaModule& PersonaModule = FModuleManager::LoadModuleChecked<FPersonaModule>("Persona");
+						
+				FToolMenuEntry WidgetEntry = FToolMenuEntry::InitWidget(
+					"CurvePicker",
+					PersonaModule.CreateCurvePicker(
+						SetBinding->GetSkeleton(),
+						FOnCurvePicked::CreateLambda([this](const FName& InCurve)
+							{
+								const FScopedTransaction Transaction(LOCTEXT("AddCurvesToSet", "Add Curve to Set"));
+								SetBinding->Modify();
+
+								FAnimationAttributeIdentifier CurveAttrId(InCurve, INDEX_NONE, NAME_None, FFloatAnimationAttribute::StaticStruct());
+								SetBinding->AddAttributeToSet(CurveAttrId, NAME_None);
+
+								FSlateApplication::Get().DismissAllMenus();
+								RepopulateListData();
+							})),
+					FText(),
+					true,
+					false,
+					true
+				);
+
+				Section.AddEntry(WidgetEntry);
+			}));
+	}
+}
+
+TSharedRef<SWidget> SAttributesList::CreateAddAttributeWidget()
+{
+	FToolMenuContext MenuContext;
+
+	return UToolMenus::Get()->GenerateWidget(FMenus::AddAttributeId, MenuContext);
+}
+
+#undef LOCTEXT_NAMESPACE
\ No newline at end of file
diff --git a/Engine/Plugins/Experimental/Animation/SkeletonTemplateFramework/Source/Editor/Private/SkeletonTemplateFramework/Sets/SAttributesList.h b/Engine/Plugins/Experimental/Animation/SkeletonTemplateFramework/Source/Editor/Private/SkeletonTemplateFramework/Sets/SAttributesList.h
new file mode 100644
--- /dev/null
+++ b/Engine/Plugins/Experimental/Animation/SkeletonTemplateFramework/Source/Editor/Private/SkeletonTemplateFramework/Sets/SAttributesList.h
@@ -0,0 +1,70 @@
+// Copyright Epic Games, Inc. All Rights Reserved.
+
+#pragma once
+
+#include "Animation/AnimData/AttributeIdentifier.h"
+#include "CoreMinimal.h"
+#include "Widgets/SCompoundWidget.h"
+#include "Widgets/Views/SListView.h"
+
+class UAbstractSkeletonSetBinding;
+
+namespace UE::Anim::SkelTemplates::Sets
+{
+
+	class SAttributesList : public SCompoundWidget
+	{
+	public:
+		SLATE_BEGIN_ARGS(SAttributesList) {}
+			SLATE_ARGUMENT(TWeakObjectPtr<UAbstractSkeletonSetBinding>, SetBinding)
+			SLATE_EVENT(FSimpleDelegate, OnListRefreshed)
+		SLATE_END_ARGS()
+
+		virtual void Construct(const FArguments& InArgs);
+
+		void RepopulateListData();
+
+		void SetSetBinding(TWeakObjectPtr<UAbstractSkeletonSetBinding> InSetBinding);
+
+	public:
+		struct FColumns
+		{
+			static FName SetId;
+			static FName AttributeId;
+		};
+
+		struct FMenus
+		{
+			static FName AddAttributeId;
+		};
+
+		struct FListItem
+		{
+			FListItem(const FName InSetName, const FAnimationAttributeIdentifier& InAttribute);
+
+			FAnimationAttributeIdentifier Attribute;
+			FName SetName;
+		};
+
+		using FListItemPtr = TSharedPtr<FListItem>;
+
+	private:
+		TSharedRef<ITableRow> ListView_OnGenerateRow(FListItemPtr Item, const TSharedRef<STableViewBase>& OwnerTable);
+		
+		TSharedRef<SWidget> CreateAddAttributeWidget();
+
+		void RegisterMenus();
+
+	private:
+		TArray<FListItemPtr> ListItems;
+
+		TSharedPtr<SListView<FListItemPtr>> ListView;
+
+		TWeakObjectPtr<UAbstractSkeletonSetBinding> SetBinding;
+
+		FSimpleDelegate OnListRefreshed;
+
+		bool bRepopulating = false;
+	};
+
+}
\ No newline at end of file
diff --git a/Engine/Plugins/Experimental/Animation/SkeletonTemplateFramework/Source/Editor/Private/SkeletonTemplateFramework/Sets/SSetBinding.cpp b/Engine/Plugins/Experimental/Animation/SkeletonTemplateFramework/Source/Editor/Private/SkeletonTemplateFramework/Sets/SSetBinding.cpp
new file mode 100644
--- /dev/null
+++ b/Engine/Plugins/Experimental/Animation/SkeletonTemplateFramework/Source/Editor/Private/SkeletonTemplateFramework/Sets/SSetBinding.cpp
@@ -0,0 +1,448 @@
+// Copyright Epic Games, Inc. All Rights Reserved.
+
+#include "SkeletonTemplateFramework/Sets/SSetBinding.h"
+
+#include "Framework/Commands/UICommandList.h"
+#include "Framework/MultiBox/MultiBoxBuilder.h"
+#include "ScopedTransaction.h"
+#include "SetDragDrop.h"
+#include "SkeletonTreeDragDrop.h"
+#include "SkeletonTemplateFramework/AbstractSkeletonSetBinding.h"
+#include "SkeletonTemplateFramework/AbstractSkeletonSetCollection.h"
+#include "Widgets/Input/SSearchBox.h"
+
+#define LOCTEXT_NAMESPACE "UE::Anim::SkelTemplates::Sets::SSetBinding"
+
+namespace UE::Anim::SkelTemplates::Sets
+{
+void SSetBinding::Construct(const FArguments& InArgs, TWeakObjectPtr<UAbstractSkeletonSetBinding> InSetBinding)
+	{
+		SetBinding = InSetBinding;
+		OnTreeRefreshed = InArgs._OnTreeRefreshed;
+
+		ChildSlot
+			[
+				SNew(SVerticalBox)
+					+ SVerticalBox::Slot()
+					.AutoHeight()
+					.Padding(FMargin(1.0f, 1.0f))
+					[
+						SNew(SHorizontalBox)
+							+ SHorizontalBox::Slot()
+							.FillWidth(1.0f)
+							.Padding(FMargin(1.0f, 1.0f))
+							[
+								SAssignNew(SearchBox, SSearchBox)
+									.SelectAllTextWhenFocused(true)
+									.OnTextChanged_Lambda([](const FText& InText) { /* TODO: Implement */ })
+									.HintText(LOCTEXT("SearchBox_Hint", "Search Sets..."))
+							]
+					]
+					+ SVerticalBox::Slot()
+					.FillHeight(1.0f)
+					.Padding(FMargin(1.0f, 1.0f))
+					[
+						SAssignNew(TreeView, STreeView<TSharedPtr<ITreeItem>>)
+							.SelectionMode(ESelectionMode::Single)
+							.HighlightParentNodesForSelection(true)
+							.TreeItemsSource(&RootItems)
+							.OnGenerateRow(this, &SSetBinding::TreeView_OnGenerateRow)
+							.OnGetChildren(this, &SSetBinding::TreeView_OnGetChildren)
+							.OnContextMenuOpening(this, &SSetBinding::TreeView_OnContextMenuOpening)
+							.HeaderRow
+							(
+								SNew(SHeaderRow)
+								+ SHeaderRow::Column(FName("SetHierarchy"))
+								.FillWidth(0.5f)
+								.DefaultLabel(LOCTEXT("TreeView_Header", "Set Bindings"))
+							)
+					]
+			];
+
+		RepopulateTreeData();
+		ExpandAllTreeItems();
+		BindCommands();
+	}
+
+	FReply SSetBinding::OnKeyDown(const FGeometry& MyGeometry, const FKeyEvent& InKeyEvent)
+	{
+		if (CommandList->ProcessCommandBindings(InKeyEvent))
+		{
+			return FReply::Handled();
+		}
+
+		return FReply::Unhandled();
+	}
+
+	TSharedRef<ITableRow> SSetBinding::TreeView_OnGenerateRow(TSharedPtr<ITreeItem> InItem, const TSharedRef<STableViewBase>& OwnerTable)
+	{
+		return InItem->GenerateRow(OwnerTable);
+	}
+
+	void SSetBinding::TreeView_OnGetChildren(TSharedPtr<ITreeItem> InItem, TArray<TSharedPtr<ITreeItem>>& OutChildren)
+	{
+		InItem->GetChildren(OutChildren);
+	}
+
+	TSharedPtr<SWidget> SSetBinding::TreeView_OnContextMenuOpening()
+	{
+		FMenuBuilder MenuBuilder(true, CommandList.ToSharedRef());
+
+		TArray<TSharedPtr<ITreeItem>> Selection = TreeView->GetSelectedItems();
+		if (Selection.Num() > 0)
+		{
+			TSharedPtr<ITreeItem> Selected = Selection[0]; // Single selection only
+
+			const bool bShowMenu = Selected->GenerateContextMenu(*this, MenuBuilder);
+
+			if (bShowMenu)
+			{
+				return MenuBuilder.MakeWidget();
+			}
+		}
+
+		return SNullWidget::NullWidget;
+	}
+
+	void SSetBinding::PostUndo(bool bSuccess)
+	{
+		RepopulateTreeData();
+	}
+
+	void SSetBinding::PostRedo(bool bSuccess)
+	{
+		RepopulateTreeData();
+	}
+
+	void SSetBinding::SetSetBinding(TWeakObjectPtr<UAbstractSkeletonSetBinding> InSetBinding)
+	{
+		SetBinding = InSetBinding;
+		RepopulateTreeData();
+	}
+
+	void SSetBinding::RepopulateTreeData()
+	{
+		if (!bRepopulating)
+		{
+			TGuardValue<bool> RecursionGuard(bRepopulating, true);
+
+			// Save the current expanded state of each tree item
+			// - Only Set items can contain children so only Set items can be expanded/collapsed
+			TSet<FName> ClosedSetCollections;
+			{
+				for (TSharedPtr<ITreeItem> TreeItem : GetAllTreeItems())
+				{
+					if (TreeItem->GetType() == FTreeItem_Set::StaticGetType() && !TreeView->IsItemExpanded(TreeItem))
+					{
+						TSharedPtr<FTreeItem_Set> SetTreeItem = StaticCastSharedPtr<FTreeItem_Set>(TreeItem);
+						ClosedSetCollections.Add(SetTreeItem->Name);
+					}
+				}
+			}
+			
+			// Rebuild tree data
+			RootItems.Reset();
+
+			if (SetBinding.IsValid())
+			{
+				TMap<FName, TSharedPtr<FTreeItem_Set>> TreeItemMap;
+				
+				if (SetBinding->GetSetCollection())
+				{
+					for (const FAbstractSkeletonSet& Set : SetBinding->GetSetCollection()->GetSetHierarchy())
+					{
+						TSharedPtr<FTreeItem_Set> SetItem = MakeShared<FTreeItem_Set>();
+						SetItem->Name = Set.SetName;
+						SetItem->ParentName = Set.ParentSetName;
+						SetItem->OnCanAcceptDrop = STableRow<TSharedPtr<ITreeItem>>::FOnCanAcceptDrop::CreateLambda([](const FDragDropEvent& DragDropEvent, EItemDropZone DropZone, TSharedPtr<ITreeItem> TargetItem) -> TOptional<EItemDropZone>
+							{
+								using namespace UE::Anim::SkelTemplates;
+
+								TOptional<EItemDropZone> ReturnedDropZone;
+
+								const TSharedPtr<FSkeletonTreeDragDropOp> DragDropOp = DragDropEvent.GetOperationAs<FSkeletonTreeDragDropOp>();
+								if (DragDropOp.IsValid())
+								{
+									ReturnedDropZone = EItemDropZone::OntoItem;
+								}
+
+								return ReturnedDropZone;
+							});
+						SetItem->OnAcceptDrop = STableRow<TSharedPtr<ITreeItem>>::FOnAcceptDrop::CreateLambda([this](const FDragDropEvent& DragDropEvent, EItemDropZone DropZone, TSharedPtr<ITreeItem> TargetTreeItem) -> FReply
+							{
+								const TSharedPtr<FSkeletonTreeDragDropOp> DragDropOp = DragDropEvent.GetOperationAs<FSkeletonTreeDragDropOp>();
+								if (DragDropOp.IsValid())
+								{
+									TSharedPtr<FTreeItem_Set> TargetSetTreeItem = StaticCastSharedPtr<FTreeItem_Set>(TargetTreeItem);
+
+									const FScopedTransaction Transaction(LOCTEXT("AddAttributeToSet", "Add Attribute to Set"));
+									SetBinding->Modify();
+
+									const bool bSuccess = SetBinding->AddBoneToSet(DragDropOp->Item->BoneName, TargetSetTreeItem->Name);
+									check(bSuccess);
+									RepopulateTreeData();
+
+									return FReply::Handled();
+								}
+								return FReply::Unhandled();
+							});
+
+						if (Set.ParentSetName == NAME_None)
+						{
+							RootItems.Add(SetItem);
+						}
+						else
+						{
+							if (!TreeItemMap.Contains(Set.ParentSetName))
+							{
+								TSharedPtr<FTreeItem_Set> ParentItem = MakeShared<FTreeItem_Set>();
+								ParentItem->Name = Set.ParentSetName;
+
+								TreeItemMap.Add(Set.ParentSetName, ParentItem);
+							}
+
+							TSharedPtr<FTreeItem_Set> ParentItem = TreeItemMap[Set.ParentSetName];
+							ParentItem->Children.Add(SetItem);
+
+						}
+
+						TreeItemMap.Add(Set.SetName, SetItem);
+					}
+				}
+
+				for (const FAbstractSkeleton_BoneBinding& Binding : SetBinding->GetBoneBindings())
+				{
+					TSharedPtr<FTreeItem_Bone> BindingItem = MakeShared<FTreeItem_Bone>();
+					BindingItem->BoneName = Binding.BoneName;
+
+					if (TreeItemMap.Contains(Binding.SetName))
+					{
+						TSharedPtr<FTreeItem_Set> SetItem = TreeItemMap[Binding.SetName];
+						SetItem->Children.Insert(BindingItem, 0);
+					}
+					else
+					{
+						check(false);
+					}
+				}
+
+				for (const FAbstractSkeleton_AttributeBinding& Binding : SetBinding->GetAttributeBindings())
+				{
+					TSharedPtr<FTreeItem_Attribute> BindingItem = MakeShared<FTreeItem_Attribute>();
+					BindingItem->Attribute = Binding.Attribute;
+
+					if (TreeItemMap.Contains(Binding.SetName))
+					{
+						TSharedPtr<FTreeItem_Set> SetItem = TreeItemMap[Binding.SetName];
+						SetItem->Children.Insert(BindingItem, 0);
+					}
+				}
+			}
+
+			TreeView->RebuildList();
+
+			// Restore expanded/collapsed state
+			{
+				for (TSharedPtr<ITreeItem> TreeItem : GetAllTreeItems())
+				{
+					if (TreeItem->GetType() == FTreeItem_Set::StaticGetType())
+					{
+						const TSharedPtr<FTreeItem_Set> SetTreeItem = StaticCastSharedPtr<FTreeItem_Set>(TreeItem);
+						const bool bShouldExpand = !ClosedSetCollections.Contains(SetTreeItem->Name);
+						TreeView->SetItemExpansion(TreeItem, bShouldExpand);
+					}
+				}
+			}
+
+			OnTreeRefreshed.ExecuteIfBound();
+		}
+	}
+
+	void SSetBinding::ExpandAllTreeItems()
+	{
+		for (TSharedPtr<ITreeItem> TreeItem : GetAllTreeItems())
+		{
+			TreeView->SetItemExpansion(TreeItem, true);
+		}
+	}
+
+	void SSetBinding::BindCommands()
+	{
+		CommandList = MakeShareable(new FUICommandList);
+	}
+
+	TArray<TSharedPtr<SSetBinding::ITreeItem>> SSetBinding::GetAllTreeItems()
+	{
+		TArray<TSharedPtr<ITreeItem>> AllItems;
+		AllItems.Append(RootItems);
+
+		for (int32 Index = 0; Index < AllItems.Num(); ++Index)
+		{
+			AllItems[Index]->GetChildren(AllItems);
+		}
+
+		return AllItems;
+	}
+	
+	void SSetBinding::FTreeItem_Set::GetChildren(TArray<TSharedPtr<ITreeItem>>& OutChildren)
+	{
+		OutChildren.Append(Children);
+	}
+
+	TSharedRef<ITableRow> SSetBinding::FTreeItem_Set::GenerateRow(const TSharedRef<STableViewBase>& OwnerTable)
+	{
+		return SNew(STableRow<TSharedPtr<ITreeItem>>, OwnerTable)
+			.Style(FAppStyle::Get(), "TableView.AlternatingRow")
+			.OnDragDetected_Lambda([this](const FGeometry& MyGeometry, const FPointerEvent& MouseEvent)
+				{
+					if (MouseEvent.IsMouseButtonDown(EKeys::LeftMouseButton))
+					{
+						const TSharedRef<FSetDragDropOp> DragDropOp = FSetDragDropOp::New(this);
+						return FReply::Handled().BeginDragDrop(DragDropOp);
+					}
+
+					return FReply::Unhandled();
+				})
+			.OnCanAcceptDrop(OnCanAcceptDrop)
+			.OnAcceptDrop(OnAcceptDrop)
+			.ShowWires(true)
+			[
+				SNew(SHorizontalBox)
+					+ SHorizontalBox::Slot()
+					.AutoWidth()
+					.Padding(FMargin(3.0f))
+					[
+						SNew(SImage)
+							.Image(FAppStyle::Get().GetBrush("ClassIcon.GroupActor"))
+					]
+					+ SHorizontalBox::Slot()
+					.FillWidth(1.0)
+					.AutoWidth()
+					.Padding(3.0f)
+					[
+						SNew(STextBlock)
+							.Text(FText::FromName(Name))
+					]
+
+			];
+	}
+
+	void SSetBinding::FTreeItem_Bone::GetChildren(TArray<TSharedPtr<ITreeItem>>& OutChildren)
+	{
+		// Bone items don't have children, do nothing
+	}
+
+	void SSetBinding::FTreeItem_Attribute::GetChildren(TArray<TSharedPtr<ITreeItem>>& OutChildren)
+	{
+		// Attribute items don't have children, do nothing
+	}
+
+	TSharedRef<ITableRow> SSetBinding::FTreeItem_Bone::GenerateRow(const TSharedRef<STableViewBase>& OwnerTable)
+	{
+		const FSlateBrush* Icon = FAppStyle::Get().GetBrush("SkeletonTree.Bone");	
+
+		return SNew(STableRow<TSharedPtr<ITreeItem>>, OwnerTable)
+			.Style(FAppStyle::Get(), "TableView.AlternatingRow")
+			.ShowWires(true)
+			[
+				SNew(SHorizontalBox)
+					+ SHorizontalBox::Slot()
+					.AutoWidth()
+					.Padding(FMargin(3.0f))
+					[
+						SNew(SImage)
+							.Image(Icon)
+							.ColorAndOpacity(FLinearColor(FColor(31, 228, 75)))
+					]
+					+ SHorizontalBox::Slot()
+					.FillWidth(1.0)
+					.AutoWidth()
+					.Padding(3.0f)
+					[
+						SNew(STextBlock)
+							.Text(FText::FromName(BoneName))
+					]
+
+			];
+	}
+
+	TSharedRef<ITableRow> SSetBinding::FTreeItem_Attribute::GenerateRow(const TSharedRef<STableViewBase>& OwnerTable)
+	{
+		const FSlateBrush* Icon = FAppStyle::Get().GetBrush("AnimGraph.Attribute.Attributes.Icon");
+
+		return SNew(STableRow<TSharedPtr<ITreeItem>>, OwnerTable)
+			.ShowWires(true)
+			[
+				SNew(SHorizontalBox)
+					+ SHorizontalBox::Slot()
+					.AutoWidth()
+					.Padding(FMargin(1.0f, 1.0f))
+					[
+						SNew(SImage)
+							.Image(Icon)
+							.ColorAndOpacity(FLinearColor(FColor(31, 228, 75)))
+					]
+					+ SHorizontalBox::Slot()
+					.FillWidth(1.0)
+					.AutoWidth()
+					.Padding(4.0f, 0.0f, 0.0f, 0.0f)
+					[
+						SNew(STextBlock)
+							.Text(FText::FromName(Attribute.GetName()))
+					]
+
+			];
+	}
+
+	bool SSetBinding::FTreeItem_Set::GenerateContextMenu(SSetBinding& SetBindingWidget, FMenuBuilder& Menu)
+	{
+		return false;
+	}
+
+	bool SSetBinding::FTreeItem_Bone::GenerateContextMenu(SSetBinding& SetBindingWidget, FMenuBuilder& Menu)
+	{
+		FUIAction UnbindBone(
+			FExecuteAction::CreateLambda([&]()
+			{
+				FScopedTransaction Transaction(LOCTEXT("RemoveBoneFromSet", "Remove Bone From Set"));
+				SetBindingWidget.SetBinding->Modify();
+
+				SetBindingWidget.SetBinding->RemoveBoneFromSet(BoneName);
+				SetBindingWidget.RepopulateTreeData();
+			})
+		);
+
+		Menu.AddMenuEntry(
+			LOCTEXT("UnbindBone_Label", "Unbind Bone"),
+			LOCTEXT("UnbindBone_Description", "Unbind the selected bone."),
+			FSlateIcon(),
+			UnbindBone);
+
+		return true;
+	}
+	
+	bool SSetBinding::FTreeItem_Attribute::GenerateContextMenu(SSetBinding& SetBindingWidget, FMenuBuilder& Menu)
+	{
+		FUIAction UnbindAttribute(
+			FExecuteAction::CreateLambda([&]()
+			{
+				FScopedTransaction Transaction(LOCTEXT("RemoveAttributeFromSet", "Remove Attribute From Set"));
+				SetBindingWidget.SetBinding->Modify();
+
+				SetBindingWidget.SetBinding->RemoveAttributeFromSet(Attribute);
+				SetBindingWidget.RepopulateTreeData();
+			})
+		);
+
+		Menu.AddMenuEntry(
+			LOCTEXT("UnbindAttribute_Label", "Unbind Attribute"),
+			LOCTEXT("UnbindAttribute_Description", "Unbind the selected attribute."),
+			FSlateIcon(),
+			UnbindAttribute);
+
+		return true;
+	}
+
+}
+
+#undef LOCTEXT_NAMESPACE
\ No newline at end of file
diff --git a/Engine/Plugins/Experimental/Animation/SkeletonTemplateFramework/Source/Editor/Private/SkeletonTemplateFramework/Sets/SSetBinding.h b/Engine/Plugins/Experimental/Animation/SkeletonTemplateFramework/Source/Editor/Private/SkeletonTemplateFramework/Sets/SSetBinding.h
new file mode 100644
--- /dev/null
+++ b/Engine/Plugins/Experimental/Animation/SkeletonTemplateFramework/Source/Editor/Private/SkeletonTemplateFramework/Sets/SSetBinding.h
@@ -0,0 +1,122 @@
+// Copyright Epic Games, Inc. All Rights Reserved.
+
+#pragma once
+
+#include "Animation/AnimData/AttributeIdentifier.h"
+#include "CoreMinimal.h"
+#include "EditorUndoClient.h"
+#include "Widgets/SCompoundWidget.h"
+#include "Widgets/Views/STableRow.h"
+#include "Widgets/Views/STreeView.h"
+
+class FUICommandList;
+class SSearchBox;
+class UAbstractSkeletonSetBinding;
+
+namespace UE::Anim::SkelTemplates::Sets
+{
+	class SSetBinding : public SCompoundWidget, public FSelfRegisteringEditorUndoClient
+	{
+		SLATE_BEGIN_ARGS(SSetBinding) {}
+			SLATE_EVENT(FSimpleDelegate, OnTreeRefreshed)
+		SLATE_END_ARGS()
+
+	public:
+		struct ITreeItem
+		{
+			virtual ~ITreeItem() = default;
+			virtual FName GetType() const = 0;
+			virtual void GetChildren(TArray<TSharedPtr<ITreeItem>>& OutChildren) = 0;
+			virtual TSharedRef<ITableRow> GenerateRow(const TSharedRef<STableViewBase>& OwnerTable) = 0;
+			virtual bool GenerateContextMenu(SSetBinding& SetBindingWidget, FMenuBuilder& Menu) = 0;
+		};
+
+		struct FTreeItem_Set : public ITreeItem
+		{
+			virtual ~FTreeItem_Set() override = default;
+			
+			static FName StaticGetType() { return FName("Set"); }
+			
+			virtual FName GetType() const override { return StaticGetType(); }
+			virtual void GetChildren(TArray<TSharedPtr<ITreeItem>>& OutChildren) override;
+			virtual TSharedRef<ITableRow> GenerateRow(const TSharedRef<STableViewBase>& OwnerTable) override;
+			virtual bool GenerateContextMenu(SSetBinding& SetBindingWidget, FMenuBuilder& Menu) override;
+
+			FName Name;
+			FName ParentName;
+			TArray<TSharedPtr<ITreeItem>> Children;
+			STableRow<TSharedPtr<ITreeItem>>::FOnCanAcceptDrop OnCanAcceptDrop;
+			STableRow<TSharedPtr<ITreeItem>>::FOnAcceptDrop OnAcceptDrop;
+		};
+
+		struct FTreeItem_Bone : public ITreeItem
+		{
+			virtual ~FTreeItem_Bone() override = default;
+
+			static FName StaticGetType() { return FName("Bone"); }
+			
+			virtual FName GetType() const override { return StaticGetType(); }
+			virtual void GetChildren(TArray<TSharedPtr<ITreeItem>>& OutChildren) override;
+			virtual TSharedRef<ITableRow> GenerateRow(const TSharedRef<STableViewBase>& OwnerTable) override;
+			virtual bool GenerateContextMenu(SSetBinding& SetBindingWidget, FMenuBuilder& Menu) override;
+
+			FName BoneName;
+		};
+		
+		struct FTreeItem_Attribute : public ITreeItem
+		{
+			virtual ~FTreeItem_Attribute() override = default;
+
+			static FName StaticGetType() { return FName("Attribute"); }
+			
+			virtual FName GetType() const override { return StaticGetType(); }
+			virtual void GetChildren(TArray<TSharedPtr<ITreeItem>>& OutChildren) override;
+			virtual TSharedRef<ITableRow> GenerateRow(const TSharedRef<STableViewBase>& OwnerTable) override;
+			virtual bool GenerateContextMenu(SSetBinding& SetBindingWidget, FMenuBuilder& Menu) override;
+
+			FAnimationAttributeIdentifier Attribute;
+		};
+
+		void Construct(const FArguments& InArgs, TWeakObjectPtr<UAbstractSkeletonSetBinding> InSetBinding);
+
+		/** FSelfRegisteringEditorUndoClient interface */
+		virtual void PostUndo(bool bSuccess) override;
+		virtual void PostRedo(bool bSuccess) override;
+
+		void RepopulateTreeData();
+
+		void SetSetBinding(TWeakObjectPtr<UAbstractSkeletonSetBinding> InSetBinding);
+
+	private:
+		virtual FReply OnKeyDown(const FGeometry& MyGeometry, const FKeyEvent& InKeyEvent) override;
+
+		TSharedRef<ITableRow> TreeView_OnGenerateRow(TSharedPtr<ITreeItem> InItem, const TSharedRef<STableViewBase>& OwnerTable);
+
+		void TreeView_OnGetChildren(TSharedPtr<ITreeItem> InItem, TArray<TSharedPtr<ITreeItem>>& OutChildren);
+
+		TSharedPtr<SWidget> TreeView_OnContextMenuOpening();
+
+		void ExpandAllTreeItems();
+
+		void BindCommands();
+
+		TArray<TSharedPtr<ITreeItem>> GetAllTreeItems();
+
+	public:
+		TWeakObjectPtr<UAbstractSkeletonSetBinding> SetBinding;
+
+	private:
+		TSharedPtr<STreeView<TSharedPtr<ITreeItem>>> TreeView;
+
+		TArray<TSharedPtr<ITreeItem>> RootItems;
+
+		TSharedPtr<SSearchBox> SearchBox;
+
+		TSharedPtr<FUICommandList> CommandList;
+
+		FSimpleDelegate OnTreeRefreshed;
+
+		bool bRepopulating = false;
+	};
+
+}
diff --git a/Engine/Plugins/Experimental/Animation/SkeletonTemplateFramework/Source/Editor/Private/SkeletonTemplateFramework/Sets/SSetSkeletonTree.cpp b/Engine/Plugins/Experimental/Animation/SkeletonTemplateFramework/Source/Editor/Private/SkeletonTemplateFramework/Sets/SSetSkeletonTree.cpp
new file mode 100644
--- /dev/null
+++ b/Engine/Plugins/Experimental/Animation/SkeletonTemplateFramework/Source/Editor/Private/SkeletonTemplateFramework/Sets/SSetSkeletonTree.cpp
@@ -0,0 +1,281 @@
+// Copyright Epic Games, Inc. All Rights Reserved.
+
+#include "SkeletonTemplateFramework/Sets/SSetSkeletonTree.h"
+
+#include "Animation/Skeleton.h"
+#include "ReferenceSkeleton.h"
+#include "ScopedTransaction.h"
+#include "SetDragDrop.h"
+#include "SkeletonTreeDragDrop.h"
+#include "SkeletonTemplateFramework/AbstractSkeletonSetBinding.h"
+#include "Widgets/Input/SSearchBox.h"
+#include "Widgets/Views/STreeView.h"
+
+#define LOCTEXT_NAMESPACE "UE::Anim::SkelTemplates::Sets::SSetsSkeletonTree"
+
+namespace UE::Anim::SkelTemplates::Sets
+{
+	void SSetsSkeletonTree::Construct(const FArguments& InArgs)
+	{
+		SetBinding = InArgs._SetBinding;
+		OnTreeRefreshed = InArgs._OnTreeRefreshed;
+
+		ChildSlot
+			[
+				SNew(SVerticalBox)
+					+ SVerticalBox::Slot()
+					.AutoHeight()
+					.Padding(FMargin(1.0f, 1.0f))
+					[
+						SNew(SHorizontalBox)
+							+ SHorizontalBox::Slot()
+							.FillWidth(1.0f)
+							.Padding(FMargin(1.0f, 1.0f))
+							[
+								SNew(SSearchBox)
+									.SelectAllTextWhenFocused(true)
+									.OnTextChanged_Lambda([](const FText& InText) { /* TODO: Implement */ })
+									.HintText(LOCTEXT("SearchBox_Hint", "Search Sets..."))
+							]
+					]
+					+ SVerticalBox::Slot()
+					[
+						SAssignNew(TreeView, STreeView<TSharedPtr<FTreeItem>>)
+							.TreeItemsSource(&RootItems)
+							.OnGenerateRow(this, &SSetsSkeletonTree::TreeView_OnGenerateRow)
+							.OnGetChildren(this, &SSetsSkeletonTree::TreeView_OnGetChildren)
+							.HeaderRow(
+								SNew(SHeaderRow)
+								+ SHeaderRow::Column("Set")
+								.DefaultLabel(LOCTEXT("SetLabel", "Set"))
+								.FillWidth(0.3f)
+								+ SHeaderRow::Column("Name")
+								.DefaultLabel(LOCTEXT("NameLabel", "Name"))
+								.FillWidth(0.7f)
+							)
+					]
+			];
+
+		RepopulateTreeData();
+	}
+
+	class SSetsSkeletonTreeRow : public SMultiColumnTableRow<TSharedPtr<SSetsSkeletonTree::FTreeItem>>
+	{
+	public:
+		SLATE_BEGIN_ARGS(SSetsSkeletonTreeRow) {}
+			SLATE_ARGUMENT(TSharedPtr<SSetsSkeletonTree::FTreeItem>, Item)
+			SLATE_EVENT(FOnDragDetected, OnDragDetected)
+			SLATE_EVENT(FOnCanAcceptDrop, OnCanAcceptDrop)
+			SLATE_EVENT(FOnAcceptDrop, OnAcceptDrop)
+		SLATE_END_ARGS()
+
+		void Construct(const FArguments& InArgs, const TSharedRef<STableViewBase>& OwnerTableView)
+		{
+			Item = InArgs._Item;
+
+			const FSuperRowType::FArguments Args = FSuperRowType::FArguments()
+				.OnDragDetected(InArgs._OnDragDetected)
+				.OnCanAcceptDrop(InArgs._OnCanAcceptDrop)
+				.OnAcceptDrop(InArgs._OnAcceptDrop)
+				.Style(FAppStyle::Get(), "TableView.AlternatingRow");
+
+			SMultiColumnTableRow<TSharedPtr<SSetsSkeletonTree::FTreeItem>>::Construct(Args, OwnerTableView);
+		}
+
+		virtual TSharedRef<SWidget> GenerateWidgetForColumn(const FName& ColumnName) override
+		{
+			if (ColumnName == "Set")
+			{
+				return SNew(SHorizontalBox)
+					+ SHorizontalBox::Slot()
+					.AutoWidth()
+					.VAlign(VAlign_Center)
+					.Padding(6, 0)
+					[
+						SNew(SBox)
+							.WidthOverride(8.0f)
+							.HeightOverride(8.0f)
+							[
+								SNew(SImage)
+									.Image(FAppStyle::Get().GetBrush("Icons.FilledCircle"))
+									.ColorAndOpacity((Item->BoundSet != NAME_None) ? FLinearColor(FColor(31, 228, 75)) : FLinearColor(FColor(239, 53, 53)))
+							]
+					]
+					+ SHorizontalBox::Slot()
+					.AutoWidth()
+					.VAlign(VAlign_Center)
+					.Padding(2, 2)
+					[
+						SNew(SImage)
+							.Visibility(Item->BoundSet == NAME_None ? EVisibility::Collapsed : EVisibility::Visible)
+							.Image(FAppStyle::Get().GetBrush("ClassIcon.GroupActor"))
+					]
+					+ SHorizontalBox::Slot()
+					.AutoWidth()
+					.VAlign(VAlign_Center)
+					.Padding(2, 2)
+					[
+						SNew(STextBlock)
+							.Visibility(Item->BoundSet == NAME_None ? EVisibility::Collapsed : EVisibility::Visible)
+							.Text(FText::FromName(Item->BoundSet))
+					];
+
+			}
+			else if (ColumnName == "Name")
+			{
+				return SNew(SHorizontalBox)
+					+ SHorizontalBox::Slot()
+					.AutoWidth()
+					[
+						SNew(SExpanderArrow, SharedThis(this))
+							.ShouldDrawWires(true)
+					]
+					+ SHorizontalBox::Slot()
+					.AutoWidth()
+					.VAlign(VAlign_Center)
+					.Padding(2, 2)
+					[
+						SNew(SImage)
+							.Image(FAppStyle::Get().GetBrush("SkeletonTree.Bone"))
+							.ColorAndOpacity((Item->BoundSet != NAME_None) ? FLinearColor(FColor(31, 228, 75)) : FSlateColor::UseForeground())
+					]
+					+ SHorizontalBox::Slot()
+					.AutoWidth()
+					.VAlign(VAlign_Center)
+					.Padding(2, 2)
+					[
+						SNew(STextBlock)
+							.Text(FText::FromName(Item->BoneName))
+					];
+			}
+
+			return SNullWidget::NullWidget;
+		}
+
+	private:
+		TSharedPtr<SSetsSkeletonTree::FTreeItem> Item;
+	};
+
+	TSharedRef<ITableRow> SSetsSkeletonTree::TreeView_OnGenerateRow(TSharedPtr<FTreeItem> Item, const TSharedRef<STableViewBase>& OwnerTable)
+	{
+		return SNew(SSetsSkeletonTreeRow, OwnerTable)
+			.Item(Item)
+			.OnCanAcceptDrop_Lambda([](const FDragDropEvent& DragDropEvent, EItemDropZone DropZone, TSharedPtr<FTreeItem> TargetTreeItem) -> TOptional<EItemDropZone>
+				{
+					using namespace UE::Anim::SkelTemplates;
+
+					TOptional<EItemDropZone> ReturnedDropZone;
+
+					const TSharedPtr<FSetDragDropOp> DragDropOp = DragDropEvent.GetOperationAs<FSetDragDropOp>();
+					if (DragDropOp.IsValid())
+					{
+						ReturnedDropZone = EItemDropZone::OntoItem;
+					}
+
+					return ReturnedDropZone;
+				})
+			.OnAcceptDrop_Lambda([this](const FDragDropEvent& DragDropEvent, EItemDropZone DropZone, TSharedPtr<FTreeItem> TargetTreeItem) -> FReply
+				{
+					const TSharedPtr<FSetDragDropOp> DragDropOp = DragDropEvent.GetOperationAs<FSetDragDropOp>();
+					if (DragDropOp.IsValid())
+					{
+						TSharedPtr<FTreeItem> TargetSetTreeItem = StaticCastSharedPtr<FTreeItem>(TargetTreeItem);
+
+						const FScopedTransaction Transaction(LOCTEXT("AddAttributeToSet", "Add Attribute to Set"));
+						SetBinding->Modify();
+
+						const bool bSuccess = SetBinding->AddBoneToSet(TargetSetTreeItem->BoneName, DragDropOp->Item->Name);
+						check(bSuccess);
+						RepopulateTreeData();
+
+						return FReply::Handled();
+					}
+					return FReply::Unhandled();
+				})
+			.OnDragDetected_Lambda([Item](const FGeometry& MyGeometry, const FPointerEvent& MouseEvent)
+				{
+					if (MouseEvent.IsMouseButtonDown(EKeys::LeftMouseButton))
+					{
+						const TSharedRef<FSkeletonTreeDragDropOp> DragDropOp = FSkeletonTreeDragDropOp::New(Item);
+						return FReply::Handled().BeginDragDrop(DragDropOp);
+					}
+
+					return FReply::Unhandled();
+				});
+	}
+
+	void SSetsSkeletonTree::TreeView_OnGetChildren(TSharedPtr<FTreeItem> InParent, TArray<TSharedPtr<FTreeItem>>& OutChildren)
+	{
+		OutChildren = InParent->Children;
+	}
+
+	TArray<TSharedPtr<SSetsSkeletonTree::FTreeItem>> SSetsSkeletonTree::GetAllTreeItems()
+	{
+		TArray<TSharedPtr<FTreeItem>> AllItems;
+		AllItems.Append(RootItems);
+
+		for (int32 Index = 0; Index < AllItems.Num(); ++Index)
+		{
+			AllItems.Append(AllItems[Index]->Children);
+		}
+
+		return AllItems;
+	}
+
+	void SSetsSkeletonTree::SetSetBinding(TWeakObjectPtr<UAbstractSkeletonSetBinding> InSetBinding)
+	{
+		SetBinding = InSetBinding;
+		RepopulateTreeData();
+	}
+
+	void SSetsSkeletonTree::RepopulateTreeData()
+	{
+		if (!bRepopulating)
+		{
+			TGuardValue<bool> RecursionGuard(bRepopulating, true);
+
+			RootItems.Empty();
+
+			if (SetBinding.IsValid() && SetBinding->GetSkeleton())
+			{
+				const FReferenceSkeleton& RefSkeleton = SetBinding->GetSkeleton()->GetReferenceSkeleton();
+
+				using BoneFunctionType = TFunctionRef<TSharedPtr<FTreeItem>(const int32, TSharedPtr<FTreeItem>)>;
+
+				const BoneFunctionType AddBoneToTree = [&, this](const int32 BoneIndex, TSharedPtr<FTreeItem> Parent) -> TSharedPtr<FTreeItem>
+					{
+						TSharedPtr<FTreeItem> TreeItem = MakeShared<FTreeItem>();
+						TreeItem->BoneName = RefSkeleton.GetBoneName(BoneIndex);
+						TreeItem->BoundSet = SetBinding->GetBoneSet(TreeItem->BoneName);
+						TreeItem->Parent = Parent;
+
+						TArray<int32> Children;
+						RefSkeleton.GetDirectChildBones(BoneIndex, Children);
+
+						for (const int32 ChildBoneIndex : Children)
+						{
+							TSharedPtr<FTreeItem> ChildTreeItem = AddBoneToTree(ChildBoneIndex, TreeItem);
+							TreeItem->Children.Add(ChildTreeItem);
+						}
+
+						return TreeItem;
+					};
+
+				TSharedPtr<FTreeItem> RootBoneItem = AddBoneToTree(0, nullptr);
+
+				RootItems.Add(RootBoneItem);
+			}
+
+			TreeView->RequestTreeRefresh();
+
+			for (TSharedPtr<FTreeItem> TreeItem : GetAllTreeItems())
+			{
+				TreeView->SetItemExpansion(TreeItem, true);
+			}
+
+			OnTreeRefreshed.ExecuteIfBound();
+		}
+	}
+}
+
+#undef LOCTEXT_NAMESPACE
\ No newline at end of file
diff --git a/Engine/Plugins/Experimental/Animation/SkeletonTemplateFramework/Source/Editor/Private/SkeletonTemplateFramework/Sets/SSetSkeletonTree.h b/Engine/Plugins/Experimental/Animation/SkeletonTemplateFramework/Source/Editor/Private/SkeletonTemplateFramework/Sets/SSetSkeletonTree.h
new file mode 100644
--- /dev/null
+++ b/Engine/Plugins/Experimental/Animation/SkeletonTemplateFramework/Source/Editor/Private/SkeletonTemplateFramework/Sets/SSetSkeletonTree.h
@@ -0,0 +1,56 @@
+// Copyright Epic Games, Inc. All Rights Reserved.
+
+#pragma once
+
+#include "CoreMinimal.h"
+#include "Widgets/SCompoundWidget.h"
+
+class ITableRow;
+class STableViewBase;
+template <typename> class STreeView;
+class UAbstractSkeletonSetBinding;
+
+namespace UE::Anim::SkelTemplates::Sets
+{
+	class SSetsSkeletonTree : public SCompoundWidget
+	{
+	public:
+		struct FTreeItem
+		{
+			FName BoneName;
+			FName BoundSet;
+
+			TWeakPtr<FTreeItem> Parent;
+			TArray<TSharedPtr<FTreeItem>> Children;
+		};
+
+		SLATE_BEGIN_ARGS(SSetsSkeletonTree) {}
+			SLATE_ARGUMENT(TWeakObjectPtr<UAbstractSkeletonSetBinding>, SetBinding)
+			SLATE_EVENT(FSimpleDelegate, OnTreeRefreshed)
+		SLATE_END_ARGS()
+
+		virtual void Construct(const FArguments& InArgs);
+
+		void RepopulateTreeData();
+
+		void SetSetBinding(TWeakObjectPtr<UAbstractSkeletonSetBinding> InSetBinding);
+
+	private:
+		TSharedRef<ITableRow> TreeView_OnGenerateRow(TSharedPtr<FTreeItem> Item, const TSharedRef<STableViewBase>& OwnerTable);
+		void TreeView_OnGetChildren(TSharedPtr<FTreeItem> InParent, TArray<TSharedPtr<FTreeItem>>& OutChildren);
+
+		TArray<TSharedPtr<FTreeItem>> GetAllTreeItems();
+
+	private:
+		TWeakObjectPtr<UAbstractSkeletonSetBinding> SetBinding;
+
+		TArray<TSharedPtr<FTreeItem>> RootItems;
+
+		TSharedPtr<STreeView<TSharedPtr<FTreeItem>>> TreeView;
+
+		FSimpleDelegate OnTreeRefreshed;
+
+		bool bRepopulating = false;
+	};
+
+}
\ No newline at end of file
diff --git a/Engine/Plugins/Experimental/Animation/SkeletonTemplateFramework/Source/Editor/Private/SkeletonTemplateFramework/Sets/SSetsTab.cpp b/Engine/Plugins/Experimental/Animation/SkeletonTemplateFramework/Source/Editor/Private/SkeletonTemplateFramework/Sets/SSetsTab.cpp
new file mode 100644
--- /dev/null
+++ b/Engine/Plugins/Experimental/Animation/SkeletonTemplateFramework/Source/Editor/Private/SkeletonTemplateFramework/Sets/SSetsTab.cpp
@@ -0,0 +1,312 @@
+// Copyright Epic Games, Inc. All Rights Reserved.
+
+#include "SkeletonTemplateFramework/Sets/SSetsTab.h"
+
+#include "Animation/Skeleton.h"
+#include "AssetThumbnail.h"
+#include "Framework/Docking/LayoutService.h"
+#include "Framework/MultiBox/MultiBoxBuilder.h"
+#include "PropertyCustomizationHelpers.h"
+#include "SkeletonTemplateFramework/AbstractSkeletonSetBinding.h"
+#include "SkeletonTemplateFramework/AbstractSkeletonSetCollection.h"
+#include "SkeletonTemplateFramework/Sets/SSetSkeletonTree.h"
+#include "SkeletonTemplateFramework/Sets/SSetBinding.h"
+#include "SkeletonTemplateFramework/Sets/SAttributesList.h"
+#include "Widgets/Docking/SDockTab.h"
+
+#define LOCTEXT_NAMESPACE "UE::Anim::SkelTemplates::Labels::SSetsTab"
+
+namespace UE::Anim::SkelTemplates::Sets
+{
+	namespace Tabs
+	{
+		static const FName SkeletonTreeId(TEXT("AbstractSkeletonEditor_Sets_SkeletonTreeTab"));
+		static const FName AttributesId(TEXT("AbstractSkeletonEditor_Sets_AttributesTab"));
+		static const FName SetBindingsId(TEXT("AbstractSkeletonEditor_Sets_BindingsTab"));
+	}
+
+	void SSetsTab::Construct(const FArguments& InArgs)
+	{
+		SetBinding = InArgs._SetBinding;
+
+		static const FName LayoutName = "UE::Anim::STF::SSetsTab_Layout";
+
+		TSharedRef<FTabManager::FLayout> Layout = FTabManager::NewLayout(LayoutName)
+			->AddArea
+			(
+				FTabManager::NewPrimaryArea()
+				->SetOrientation(Orient_Horizontal)
+				->Split
+				(
+					FTabManager::NewStack()
+					->AddTab(Tabs::SkeletonTreeId, ETabState::OpenedTab)
+					->AddTab(Tabs::AttributesId, ETabState::OpenedTab)
+					->SetForegroundTab(Tabs::SkeletonTreeId)
+				)
+				->Split
+				(
+					FTabManager::NewStack()
+					->AddTab(Tabs::SetBindingsId, ETabState::OpenedTab)
+				)
+			);
+
+		check(InArgs._ParentTab.IsValid());
+		TabManager = FGlobalTabmanager::Get()->NewTabManager(InArgs._ParentTab.ToSharedRef());
+		TabManager->SetOnPersistLayout(FTabManager::FOnPersistLayout::CreateSP(this, &SSetsTab::HandleTabManagerPersistLayout));
+
+		TabManager->RegisterTabSpawner(Tabs::SkeletonTreeId, FOnSpawnTab::CreateLambda([this](const FSpawnTabArgs& Args)
+			{
+				return SNew(SDockTab)
+					.CanEverClose(false)
+					[
+						SAssignNew(SkeletonTreeWidget, SSetsSkeletonTree)
+							.SetBinding(SetBinding)
+							.OnTreeRefreshed_Lambda([this]()
+								{
+									if (SetBindingWidget.IsValid())
+									{
+										SetBindingWidget->RepopulateTreeData();
+									}
+
+									if (AttributesListWidget.IsValid())
+									{
+										AttributesListWidget->RepopulateListData();
+									}
+								})
+
+					];
+			}))
+			.SetDisplayName(LOCTEXT("SkeletonTree_TabDisplayName", "Skeleton Tree"))
+			.SetIcon(FSlateIcon(FAppStyle::GetAppStyleSetName(), "Persona.Tabs.SkeletonTree"));
+
+		TabManager->RegisterTabSpawner(Tabs::AttributesId, FOnSpawnTab::CreateLambda([this](const FSpawnTabArgs& Args)
+			{
+				return SNew(SDockTab)
+					.CanEverClose(false)
+					[
+						SAssignNew(AttributesListWidget, SAttributesList)
+							.SetBinding(SetBinding)
+							.OnListRefreshed_Lambda([this]()
+								{
+									if (SkeletonTreeWidget.IsValid())
+									{
+										SkeletonTreeWidget->RepopulateTreeData();
+									}
+
+									if (SetBindingWidget.IsValid())
+									{
+										SetBindingWidget->RepopulateTreeData();
+									}
+								})
+					];
+			}))
+			.SetDisplayName(LOCTEXT("Attributes_TabDisplayName", "Attributes"))
+			.SetIcon(FSlateIcon(FAppStyle::GetAppStyleSetName(), "AnimGraph.Attribute.Attributes.Icon"));
+
+		TabManager->RegisterTabSpawner(Tabs::SetBindingsId, FOnSpawnTab::CreateLambda([this](const FSpawnTabArgs& Args)
+			{
+				return SNew(SDockTab)
+					.CanEverClose(false)
+					[
+						SAssignNew(SetBindingWidget, SSetBinding, SetBinding)
+							.OnTreeRefreshed_Lambda([this]()
+								{
+									if (SkeletonTreeWidget.IsValid())
+									{
+										SkeletonTreeWidget->RepopulateTreeData();
+									}
+
+									if (AttributesListWidget.IsValid())
+									{
+										AttributesListWidget->RepopulateListData();
+									}
+								})
+					];
+			}))
+			.SetDisplayName(LOCTEXT("SetBindings_TabDisplayName", "Set Bindings"))
+			.SetIcon(FSlateIcon(FAppStyle::GetAppStyleSetName(), "LevelEditor.Tabs.Outliner"));
+
+		Layout = FLayoutSaveRestore::LoadFromConfig(GEditorLayoutIni, Layout);
+
+		FSlimHorizontalToolBarBuilder ToolbarBuilder(TSharedPtr<const FUICommandList>(), FMultiBoxCustomization::None);
+
+		ToolbarBuilder.BeginSection("Common");
+		{
+			ToolbarBuilder.AddWidget(
+				SNew(SHorizontalBox)
+				+ SHorizontalBox::Slot()
+				.AutoWidth()
+				.Padding(FMargin(16.0f, 0.0f, 24.0f, 0.0f))
+				.VAlign(VAlign_Center)
+				[
+					SNew(STextBlock)
+						.Text(LOCTEXT("SetBinding", "Set Binding"))
+				]
+				+ SHorizontalBox::Slot()
+				.AutoWidth()
+				.Padding(FMargin(0.0f, 0.0f, 16.0f, 0.0f))
+				[
+					SNew(SObjectPropertyEntryBox)
+						.AllowedClass(UAbstractSkeletonSetBinding::StaticClass())
+						.ObjectPath_Lambda([&]()
+							{
+								if (SetBinding.IsValid())
+								{
+									return SetBinding->GetPathName();
+								}
+								return FString();
+							})
+						.OnObjectChanged_Lambda([this](const FAssetData& AssetData)
+							{
+								// TODO: Update the currently edited objects in the asset editor so saving and showing dirty state is correct
+								SetBinding = CastChecked<UAbstractSkeletonSetBinding>(AssetData.GetAsset());
+
+								SkeletonTreeWidget->SetSetBinding(SetBinding);
+								AttributesListWidget->SetSetBinding(SetBinding);
+								SetBindingWidget->SetSetBinding(SetBinding);
+							})
+						.AllowClear(false)
+						.DisplayUseSelected(true)
+						.DisplayBrowse(true)
+						.DisplayThumbnail(true)
+						.ThumbnailPool(MakeShared<FAssetThumbnailPool>(1))
+				],
+				NAME_None,
+				true,
+				HAlign_Left
+			);
+
+			ToolbarBuilder.AddSeparator();
+
+			ToolbarBuilder.AddWidget(
+				SNew(SHorizontalBox)
+				+ SHorizontalBox::Slot()
+				.AutoWidth()
+				.Padding(FMargin(16.0f, 0.0f, 24.0f, 0.0f))
+				.VAlign(VAlign_Center)
+				[
+					SNew(STextBlock)
+						.Text(LOCTEXT("SetCollection", "Set Collection"))
+				]
+				+ SHorizontalBox::Slot()
+				.AutoWidth()
+				[
+					SNew(SObjectPropertyEntryBox)
+						.AllowedClass(UAbstractSkeletonSetCollection::StaticClass())
+						.OnIsEnabled_Lambda([&]()
+							{
+								return SetBinding.IsValid();
+							})
+						.ObjectPath_Lambda([&]()
+							{
+								if (SetBinding.IsValid() && SetBinding->GetSetCollection())
+								{
+									return SetBinding->GetSetCollection()->GetPathName();
+								}
+								return FString();
+							})
+						.OnObjectChanged_Lambda([this](const FAssetData& AssetData)
+							{
+								if (SetBinding.IsValid())
+								{
+									const bool bSuccess = SetBinding->SetSetCollection(CastChecked<UAbstractSkeletonSetCollection>(AssetData.GetAsset()));
+									check(bSuccess);
+
+									SetBindingWidget->RepopulateTreeData();
+									AttributesListWidget->RepopulateListData();
+									SetBindingWidget->RepopulateTreeData();
+								}
+							})
+						.AllowClear(false)
+						.DisplayUseSelected(true)
+						.DisplayBrowse(true)
+						.DisplayThumbnail(true)
+						.ThumbnailPool(MakeShared<FAssetThumbnailPool>(1))
+				],
+				NAME_None,
+				true,
+				HAlign_Left
+			);
+
+			ToolbarBuilder.AddWidget(
+				SNew(SHorizontalBox)
+				+ SHorizontalBox::Slot()
+				.AutoWidth()
+				.Padding(FMargin(16.0f, 0.0f, 24.0f, 0.0f))
+				.VAlign(VAlign_Center)
+				[
+					SNew(STextBlock)
+						.Text(LOCTEXT("Skeleton", "Skeleton"))
+				]
+				+ SHorizontalBox::Slot()
+				.AutoWidth()
+				[
+					SNew(SObjectPropertyEntryBox)
+						.AllowedClass(USkeleton::StaticClass())
+						.OnIsEnabled_Lambda([&]()
+							{
+								return SetBinding.IsValid();
+							})
+						.ObjectPath_Lambda([&]()
+							{
+								if (SetBinding.IsValid() && SetBinding->GetSkeleton())
+								{
+									return SetBinding->GetSkeleton()->GetPathName();
+								}
+								return FString();
+							})
+						.OnObjectChanged_Lambda([this](const FAssetData& AssetData)
+							{
+								if (SetBinding.IsValid())
+								{
+									const bool bSuccess = SetBinding->SetSkeleton(CastChecked<USkeleton>(AssetData.GetAsset()));
+									check(bSuccess);
+
+									SetBindingWidget->RepopulateTreeData();
+									AttributesListWidget->RepopulateListData();
+									SetBindingWidget->RepopulateTreeData();
+								}
+							})
+						.AllowClear(false)
+						.DisplayUseSelected(true)
+						.DisplayBrowse(true)
+						.DisplayThumbnail(true)
+						.ThumbnailPool(MakeShared<FAssetThumbnailPool>(1))
+				],
+				NAME_None,
+				true,
+				HAlign_Left
+			);
+		}
+		ToolbarBuilder.EndSection();
+
+		ToolbarBuilder.SetStyle(&FAppStyle::Get(), "SlimToolBar");
+
+		ChildSlot
+			[
+				SNew(SVerticalBox)
+					+ SVerticalBox::Slot()
+					.AutoHeight()
+					[
+						SNew(SHorizontalBox)
+							+ SHorizontalBox::Slot()
+							[
+								ToolbarBuilder.MakeWidget()
+							]
+					]
+					+ SVerticalBox::Slot()
+					.FillHeight(1.0f)
+					[
+						TabManager->RestoreFrom(Layout, nullptr).ToSharedRef()
+					]
+			];
+	}
+
+	void SSetsTab::HandleTabManagerPersistLayout(const TSharedRef<FTabManager::FLayout>& LayoutToSave)
+	{
+		FLayoutSaveRestore::SaveToConfig(GEditorLayoutIni, LayoutToSave);
+	}
+
+}
+
+#undef LOCTEXT_NAMESPACE
\ No newline at end of file
diff --git a/Engine/Plugins/Experimental/Animation/SkeletonTemplateFramework/Source/Editor/Private/SkeletonTemplateFramework/Sets/SSetsTab.h b/Engine/Plugins/Experimental/Animation/SkeletonTemplateFramework/Source/Editor/Private/SkeletonTemplateFramework/Sets/SSetsTab.h
new file mode 100644
--- /dev/null
+++ b/Engine/Plugins/Experimental/Animation/SkeletonTemplateFramework/Source/Editor/Private/SkeletonTemplateFramework/Sets/SSetsTab.h
@@ -0,0 +1,40 @@
+// Copyright Epic Games, Inc. All Rights Reserved.
+
+#pragma once
+
+#include "CoreMinimal.h"
+#include "Framework/Docking/TabManager.h"
+#include "Widgets/SCompoundWidget.h"
+
+class SDockTab;
+class UAbstractSkeletonSetBinding;
+
+namespace UE::Anim::SkelTemplates::Sets
+{
+	class SSetsSkeletonTree;
+	class SAttributesList;
+	class SSetBinding;
+	
+	class SSetsTab : public SCompoundWidget
+	{
+	public:
+		SLATE_BEGIN_ARGS(SSetsTab) {}
+			SLATE_ARGUMENT(TSharedPtr<SDockTab>, ParentTab)
+			SLATE_ARGUMENT(TWeakObjectPtr<UAbstractSkeletonSetBinding>, SetBinding)
+		SLATE_END_ARGS()
+
+		virtual void Construct(const FArguments& InArgs);
+
+	private:
+		void HandleTabManagerPersistLayout(const TSharedRef<FTabManager::FLayout>& LayoutToSave);
+
+	private:
+		TWeakObjectPtr<UAbstractSkeletonSetBinding> SetBinding;
+
+		TSharedPtr<SSetsSkeletonTree> SkeletonTreeWidget;
+		TSharedPtr<SAttributesList> AttributesListWidget;
+		TSharedPtr<SSetBinding> SetBindingWidget;
+
+		TSharedPtr<FTabManager> TabManager;
+	};
+}
\ No newline at end of file
diff --git a/Engine/Plugins/Experimental/Animation/SkeletonTemplateFramework/Source/Editor/Private/SkeletonTemplateFramework/Sets/SetBindingAssetDefinition.cpp b/Engine/Plugins/Experimental/Animation/SkeletonTemplateFramework/Source/Editor/Private/SkeletonTemplateFramework/Sets/SetBindingAssetDefinition.cpp
new file mode 100644
--- /dev/null
+++ b/Engine/Plugins/Experimental/Animation/SkeletonTemplateFramework/Source/Editor/Private/SkeletonTemplateFramework/Sets/SetBindingAssetDefinition.cpp
@@ -0,0 +1,43 @@
+// Copyright Epic Games, Inc. All Rights Reserved.
+
+#include "SkeletonTemplateFramework/Sets/SetBindingAssetDefinition.h"
+
+#include "SkeletonTemplateFramework/AbstractSkeletonEditor.h"
+#include "SkeletonTemplateFramework/AbstractSkeletonSetBinding.h"
+
+#define LOCTEXT_NAMESPACE "UE::Anim::SkelTemplates"
+
+FText UAssetDefinition_AbstractSkeletonSetBinding::GetAssetDisplayName() const
+{
+	return LOCTEXT("SetBinding", "UAF Set Binding");
+}
+
+FLinearColor UAssetDefinition_AbstractSkeletonSetBinding::GetAssetColor() const
+{
+	return FLinearColor(FColor(119, 119, 21));
+}
+
+TSoftClassPtr<UObject> UAssetDefinition_AbstractSkeletonSetBinding::GetAssetClass() const
+{
+	return UAbstractSkeletonSetBinding::StaticClass();
+}
+
+EAssetCommandResult UAssetDefinition_AbstractSkeletonSetBinding::OpenAssets(const FAssetOpenArgs& OpenArgs) const
+{
+	if (OpenArgs.OpenMethod == EAssetOpenMethod::Edit)
+	{
+		TArray<UObject*> Assets = OpenArgs.LoadObjects<UObject>();
+		
+		MakeShared<UE::Anim::SkelTemplates::FAbstractSkeletonEditor>()->InitEditor(Assets, OpenArgs.ToolkitHost);
+	}
+
+	return EAssetCommandResult::Handled;
+}
+
+TConstArrayView<FAssetCategoryPath> UAssetDefinition_AbstractSkeletonSetBinding::GetAssetCategories() const
+{
+	static const auto Categories = { FAssetCategoryPath(EAssetCategoryPaths::Animation, LOCTEXT("UAFSubMenu", "Animation Framework")) };
+	return Categories;
+}
+
+#undef LOCTEXT_NAMESPACE
diff --git a/Engine/Plugins/Experimental/Animation/SkeletonTemplateFramework/Source/Editor/Private/SkeletonTemplateFramework/Sets/SetBindingAssetDefinition.h b/Engine/Plugins/Experimental/Animation/SkeletonTemplateFramework/Source/Editor/Private/SkeletonTemplateFramework/Sets/SetBindingAssetDefinition.h
new file mode 100644
--- /dev/null
+++ b/Engine/Plugins/Experimental/Animation/SkeletonTemplateFramework/Source/Editor/Private/SkeletonTemplateFramework/Sets/SetBindingAssetDefinition.h
@@ -0,0 +1,22 @@
+// Copyright Epic Games, Inc. All Rights Reserved.
+
+#pragma once
+
+#include "AssetDefinitionDefault.h"
+
+#include "SetBindingAssetDefinition.generated.h"
+
+UCLASS()
+class UAssetDefinition_AbstractSkeletonSetBinding : public UAssetDefinitionDefault
+{
+	GENERATED_BODY()
+
+public:
+	// Begin UAssetDefinition
+	virtual FText GetAssetDisplayName() const override;
+	virtual FLinearColor GetAssetColor() const override;
+	virtual TSoftClassPtr<UObject> GetAssetClass() const override;
+	virtual EAssetCommandResult OpenAssets(const FAssetOpenArgs& OpenArgs) const override;
+	virtual TConstArrayView<FAssetCategoryPath> GetAssetCategories() const override;
+	// End UAssetDefinition
+};
diff --git a/Engine/Plugins/Experimental/Animation/SkeletonTemplateFramework/Source/Editor/Private/SkeletonTemplateFramework/Sets/SetBindingFactory.cpp b/Engine/Plugins/Experimental/Animation/SkeletonTemplateFramework/Source/Editor/Private/SkeletonTemplateFramework/Sets/SetBindingFactory.cpp
new file mode 100644
--- /dev/null
+++ b/Engine/Plugins/Experimental/Animation/SkeletonTemplateFramework/Source/Editor/Private/SkeletonTemplateFramework/Sets/SetBindingFactory.cpp
@@ -0,0 +1,30 @@
+// Copyright Epic Games, Inc. All Rights Reserved.
+
+#include "SkeletonTemplateFramework/Sets/SetBindingFactory.h"
+#include "SkeletonTemplateFramework/AbstractSkeletonSetBinding.h"
+#include "UObject/Package.h"
+
+UAbstractSkeletonSetBindingFactory::UAbstractSkeletonSetBindingFactory()
+{
+	bCreateNew = true;
+	bEditAfterNew = true;
+	SupportedClass = UAbstractSkeletonSetBinding::StaticClass();
+}
+
+bool UAbstractSkeletonSetBindingFactory::ConfigureProperties()
+{
+	return true;
+}
+
+UObject* UAbstractSkeletonSetBindingFactory::FactoryCreateNew(UClass* Class, UObject* InParent, FName Name, EObjectFlags Flags, UObject* Context, FFeedbackContext* Warn, FName CallingContext)
+{
+	EObjectFlags FlagsToUse = Flags | RF_Public | RF_Standalone | RF_Transactional | RF_LoadCompleted;
+	if (InParent == GetTransientPackage())
+	{
+		FlagsToUse &= ~RF_Standalone;
+	}
+
+	TObjectPtr<UAbstractSkeletonSetBinding> NewSetBinding = NewObject<UAbstractSkeletonSetBinding>(InParent, Class, Name, FlagsToUse);
+
+	return NewSetBinding;
+}
diff --git a/Engine/Plugins/Experimental/Animation/SkeletonTemplateFramework/Source/Editor/Private/SkeletonTemplateFramework/Sets/SetBindingFactory.h b/Engine/Plugins/Experimental/Animation/SkeletonTemplateFramework/Source/Editor/Private/SkeletonTemplateFramework/Sets/SetBindingFactory.h
new file mode 100644
--- /dev/null
+++ b/Engine/Plugins/Experimental/Animation/SkeletonTemplateFramework/Source/Editor/Private/SkeletonTemplateFramework/Sets/SetBindingFactory.h
@@ -0,0 +1,23 @@
+// Copyright Epic Games, Inc. All Rights Reserved.
+
+#pragma once
+
+#include "CoreMinimal.h"
+#include "Factories/Factory.h"
+
+#include "SetBindingFactory.generated.h"
+
+UCLASS()
+class UAbstractSkeletonSetBindingFactory : public UFactory
+{
+	GENERATED_BODY()
+
+	UAbstractSkeletonSetBindingFactory();
+
+	virtual bool ConfigureProperties() override;
+	virtual UObject* FactoryCreateNew(UClass* Class, UObject* InParent, FName Name, EObjectFlags Flags, UObject* Context, FFeedbackContext* Warn, FName CallingContext) override;
+	virtual UObject* FactoryCreateNew(UClass* Class, UObject* InParent, FName Name, EObjectFlags Flags, UObject* Context, FFeedbackContext* Warn) override
+	{
+		return FactoryCreateNew(Class, InParent, Name, Flags, Context, Warn, NAME_None);
+	}
+};
diff --git a/Engine/Plugins/Experimental/Animation/SkeletonTemplateFramework/Source/Editor/Private/SkeletonTemplateFramework/Sets/SetCollectionAssetDefinition.cpp b/Engine/Plugins/Experimental/Animation/SkeletonTemplateFramework/Source/Editor/Private/SkeletonTemplateFramework/Sets/SetCollectionAssetDefinition.cpp
new file mode 100644
--- /dev/null
+++ b/Engine/Plugins/Experimental/Animation/SkeletonTemplateFramework/Source/Editor/Private/SkeletonTemplateFramework/Sets/SetCollectionAssetDefinition.cpp
@@ -0,0 +1,42 @@
+// Copyright Epic Games, Inc. All Rights Reserved.
+
+#include "SkeletonTemplateFramework/Sets/SetCollectionAssetDefinition.h"
+
+#include "SkeletonTemplateFramework/Sets/SetCollectionEditor.h"
+#include "SkeletonTemplateFramework/AbstractSkeletonSetCollection.h"
+
+#define LOCTEXT_NAMESPACE "UE::Anim::SkelTemplates"
+
+FText UAssetDefinition_AbstractSkeletonSetCollection::GetAssetDisplayName() const
+{
+	return LOCTEXT("SetCollection", "UAF Set Collection");
+}
+
+FLinearColor UAssetDefinition_AbstractSkeletonSetCollection::GetAssetColor() const
+{
+	return FLinearColor(FColor(148, 148, 73));
+}
+
+TSoftClassPtr<UObject> UAssetDefinition_AbstractSkeletonSetCollection::GetAssetClass() const
+{
+	return UAbstractSkeletonSetCollection::StaticClass();
+}
+
+EAssetCommandResult UAssetDefinition_AbstractSkeletonSetCollection::OpenAssets(const FAssetOpenArgs& OpenArgs) const
+{
+	if (OpenArgs.OpenMethod == EAssetOpenMethod::Edit)
+	{
+		TArray<UObject*> Assets = OpenArgs.LoadObjects<UObject>();
+		MakeShared<UE::Anim::SkelTemplates::FSetCollectionEditorToolkit>()->InitEditor(Assets);
+	}
+
+	return EAssetCommandResult::Handled;
+}
+
+TConstArrayView<FAssetCategoryPath> UAssetDefinition_AbstractSkeletonSetCollection::GetAssetCategories() const
+{
+	static const auto Categories = { FAssetCategoryPath(EAssetCategoryPaths::Animation, LOCTEXT("UAFSubMenu", "Animation Framework")) };
+	return Categories;
+}
+
+#undef LOCTEXT_NAMESPACE
diff --git a/Engine/Plugins/Experimental/Animation/SkeletonTemplateFramework/Source/Editor/Private/SkeletonTemplateFramework/Sets/SetCollectionAssetDefinition.h b/Engine/Plugins/Experimental/Animation/SkeletonTemplateFramework/Source/Editor/Private/SkeletonTemplateFramework/Sets/SetCollectionAssetDefinition.h
new file mode 100644
--- /dev/null
+++ b/Engine/Plugins/Experimental/Animation/SkeletonTemplateFramework/Source/Editor/Private/SkeletonTemplateFramework/Sets/SetCollectionAssetDefinition.h
@@ -0,0 +1,22 @@
+// Copyright Epic Games, Inc. All Rights Reserved.
+
+#pragma once
+
+#include "AssetDefinitionDefault.h"
+
+#include "SetCollectionAssetDefinition.generated.h"
+
+UCLASS()
+class UAssetDefinition_AbstractSkeletonSetCollection : public UAssetDefinitionDefault
+{
+	GENERATED_BODY()
+
+public:
+	// Begin UAssetDefinition
+	virtual FText GetAssetDisplayName() const override;
+	virtual FLinearColor GetAssetColor() const override;
+	virtual TSoftClassPtr<UObject> GetAssetClass() const override;
+	virtual EAssetCommandResult OpenAssets(const FAssetOpenArgs& OpenArgs) const override;
+	virtual TConstArrayView<FAssetCategoryPath> GetAssetCategories() const override;
+	// End UAssetDefinition
+};
diff --git a/Engine/Plugins/Experimental/Animation/SkeletonTemplateFramework/Source/Editor/Private/SkeletonTemplateFramework/Sets/SetCollectionEditor.cpp b/Engine/Plugins/Experimental/Animation/SkeletonTemplateFramework/Source/Editor/Private/SkeletonTemplateFramework/Sets/SetCollectionEditor.cpp
new file mode 100644
--- /dev/null
+++ b/Engine/Plugins/Experimental/Animation/SkeletonTemplateFramework/Source/Editor/Private/SkeletonTemplateFramework/Sets/SetCollectionEditor.cpp
@@ -0,0 +1,510 @@
+// Copyright Epic Games, Inc. All Rights Reserved.
+
+#include "SkeletonTemplateFramework/Sets/SetCollectionEditor.h"
+
+#include "ToolMenus.h"
+#include "Widgets/Docking/SDockTab.h"
+#include "Modules/ModuleManager.h"
+#include "PropertyEditorModule.h"
+#include "SPositiveActionButton.h"
+#include "Widgets/Input/SSearchBox.h"
+#include "Framework/Commands/GenericCommands.h"
+#include "DragAndDrop/DecoratedDragDropOp.h"
+#include "ScopedTransaction.h"
+#include "Widgets/Text/SInlineEditableTextBlock.h"
+
+#define LOCTEXT_NAMESPACE "UE::Anim::SkelTemplates::FSetCollectionEditorToolkit"
+
+namespace UE::Anim::SkelTemplates
+{
+	void FSetCollectionEditorToolkit::InitEditor(const TArray<UObject*>& InObjects)
+	{
+		SetCollection = CastChecked<UAbstractSkeletonSetCollection>(InObjects[0]);
+
+		const TSharedRef<FTabManager::FLayout> Layout = FTabManager::NewLayout("AbstractSkeletonSetCollectionEditorToolkit")
+			->AddArea
+			(
+				FTabManager::NewPrimaryArea()->SetOrientation(Orient_Horizontal)
+				->Split
+				(
+					FTabManager::NewStack()
+					->SetSizeCoefficient(0.7f)
+					->AddTab("AbstractSkeletonSetCollection_SetTab", ETabState::OpenedTab)
+				)
+				->Split
+				(
+					FTabManager::NewStack()
+					->SetSizeCoefficient(0.3f)
+					->AddTab("AbstractSkeletonSetCollection_PropertiesTab", ETabState::OpenedTab)
+				)
+			);
+
+		FAssetEditorToolkit::InitAssetEditor(EToolkitMode::Standalone, {}, "AbstractSkeletonSetCollectionEditorToolkit", Layout, true, true, InObjects);
+	}
+
+	void FSetCollectionEditorToolkit::RegisterTabSpawners(const TSharedRef<class FTabManager>& InTabManager)
+	{
+		FAssetEditorToolkit::RegisterTabSpawners(InTabManager);
+
+		WorkspaceMenuCategory = InTabManager->AddLocalWorkspaceMenuCategory(LOCTEXT("AbstractSkeletonSetCollectionEditor", "Abstract Skeleton Set Editor"));
+
+		// Set tab
+		{
+			InTabManager->RegisterTabSpawner("AbstractSkeletonSetCollection_SetTab", FOnSpawnTab::CreateLambda([this](const FSpawnTabArgs&)
+				{
+					return SNew(SDockTab)
+						[
+							SNew(SSetCollection, SetCollection)
+						];
+				}))
+				.SetDisplayName(LOCTEXT("SetTab_DisplayName", "Abstract Skeleton Set"))
+				.SetGroup(WorkspaceMenuCategory.ToSharedRef());
+		}
+
+		// Property tab
+		{
+			InTabManager->RegisterTabSpawner("AbstractSkeletonSetCollection_PropertiesTab", FOnSpawnTab::CreateSP(this, &FSetCollectionEditorToolkit::SpawnPropertiesTab))
+				.SetDisplayName(LOCTEXT("PropertiesTab_DisplayName", "Details"))
+				.SetGroup(WorkspaceMenuCategory.ToSharedRef());
+		}
+	}
+
+	void FSetCollectionEditorToolkit::UnregisterTabSpawners(const TSharedRef<class FTabManager>& InTabManager)
+	{
+		FAssetEditorToolkit::UnregisterTabSpawners(InTabManager);
+		InTabManager->UnregisterTabSpawner("AbstractSkeletonSetCollection_SetTab");
+		InTabManager->UnregisterTabSpawner("AbstractSkeletonSetCollection_PropertiesTab");
+	}
+
+	TSharedRef<SDockTab> FSetCollectionEditorToolkit::SpawnPropertiesTab(const FSpawnTabArgs& Args)
+	{
+		FPropertyEditorModule& PropertyEditorModule = FModuleManager::GetModuleChecked<FPropertyEditorModule>("PropertyEditor");
+
+		FDetailsViewArgs DetailsViewArgs;
+		DetailsViewArgs.NameAreaSettings = FDetailsViewArgs::HideNameArea;
+		DetailsViewArgs.bAllowSearch = false;
+		DetailsViewArgs.bUpdatesFromSelection = false;
+		DetailsViewArgs.bShowOptions = false;
+		DetailsViewArgs.bShowObjectLabel = false;
+		DetailsViewArgs.bAllowFavoriteSystem = false;
+
+		TSharedRef<IDetailsView> DetailsView = PropertyEditorModule.CreateDetailView(DetailsViewArgs);
+		DetailsView->SetObjects(TArray<UObject*>{ SetCollection.Get() });
+
+		return SNew(SDockTab)
+			[
+				DetailsView
+			];
+	}
+
+	class FSetCollectionDragDropOp : public FDecoratedDragDropOp
+	{
+	public:
+		DRAG_DROP_OPERATOR_TYPE(FSetCollectionDragDropOp, FDecoratedDragDropOp)
+
+		static TSharedRef<FSetCollectionDragDropOp> New(const FName InSetName, const FName InParentSetName)
+		{
+			TSharedRef<FSetCollectionDragDropOp> Operation = MakeShared<FSetCollectionDragDropOp>();
+			Operation->SetName = InSetName;
+			Operation->ParentSetName = InParentSetName;
+			Operation->Construct();
+			return Operation;
+		}
+
+		TSharedPtr<SWidget> GetDefaultDecorator() const override
+		{
+			return SNew(SBorder)
+				.Visibility(EVisibility::Visible)
+				.BorderImage(FAppStyle::GetBrush("Menu.Background"))
+				[
+					SNew(STextBlock)
+						.Text(FText::FromName(SetName))
+				];
+		}
+
+		FName SetName;
+		FName ParentSetName;
+	};
+
+	void SSetCollection::Construct(const FArguments& InArgs, TWeakObjectPtr<UAbstractSkeletonSetCollection> InSetCollection)
+	{
+		SetCollection = InSetCollection;
+
+		ChildSlot
+		[
+			SNew(SVerticalBox)
+			+ SVerticalBox::Slot()
+			.AutoHeight()
+			.Padding(FMargin(1.0f, 1.0f))
+			[
+				SNew(SHorizontalBox)
+				+ SHorizontalBox::Slot()
+				.AutoWidth()
+				.Padding(FMargin(1.0f, 1.0f))
+				[
+					SNew(SPositiveActionButton)
+					.OnClicked(this, &SSetCollection::OnAddSet)
+					.Icon(FAppStyle::Get().GetBrush("Icons.Plus"))
+					.Text(LOCTEXT("AddButton_Text", "Add Set"))
+				]
+				+ SHorizontalBox::Slot()
+				.FillWidth(1.0f)
+				.Padding(FMargin(1.0f, 1.0f))
+				[
+					SAssignNew(SearchBox, SSearchBox)
+					.SelectAllTextWhenFocused(true)
+					.OnTextChanged_Lambda([](const FText& InText) { /* TODO: Implement */ })
+					.HintText(LOCTEXT("SearchBox_Hint", "Search Sets..."))
+				]
+			]
+			+ SVerticalBox::Slot()
+			.FillHeight(1.0f)
+			.Padding(FMargin(1.0f, 1.0f))
+			[
+				SAssignNew(TreeView, STreeView<TSharedPtr<FTreeItem>>)
+				.SelectionMode(ESelectionMode::Single)
+				.HighlightParentNodesForSelection(true)
+				.TreeItemsSource(&RootItems)
+				.OnGenerateRow(this, &SSetCollection::TreeView_OnGenerateRow)
+				.OnGetChildren(this, &SSetCollection::TreeView_OnGetChildren)
+				.OnContextMenuOpening(this, &SSetCollection::TreeView_OnContextMenuOpening)
+				.OnItemScrolledIntoView(this, &SSetCollection::TreeView_OnItemScrolledIntoView)
+				.OnMouseButtonDoubleClick(this, &SSetCollection::TreeView_OnMouseButtonDoubleClick)
+				.HeaderRow
+				(
+					SNew(SHeaderRow)
+					+ SHeaderRow::Column(FName("SetHierarchy"))
+					.FillWidth(0.5f)
+					.DefaultLabel(LOCTEXT("TreeView_Header", "Set Hierarchy"))
+				)
+			]
+		];
+
+		RepopulateTreeData();
+		ExpandAllTreeItems();
+		BindCommands();
+	}
+
+	FReply SSetCollection::OnKeyDown(const FGeometry& MyGeometry, const FKeyEvent& InKeyEvent)
+	{
+		if (CommandList->ProcessCommandBindings(InKeyEvent))
+		{
+			return FReply::Handled();
+		}
+
+		return FReply::Unhandled();
+	}
+
+	TSharedRef<ITableRow> SSetCollection::TreeView_OnGenerateRow(TSharedPtr<FTreeItem> InItem, const TSharedRef<STableViewBase>& OwnerTable)
+	{
+		TSharedPtr<SInlineEditableTextBlock> InlineWidget;
+
+		TSharedRef<STableRow<TSharedPtr<FTreeItem>>> RowWidget = SNew(STableRow<TSharedPtr<FTreeItem>>, OwnerTable)
+			.Style(&FAppStyle().GetWidgetStyle<FTableRowStyle>("TableView.AlternatingRow"))
+			.OnDragDetected_Lambda([InItem](const FGeometry& MyGeometry, const FPointerEvent& MouseEvent)
+				{
+					if (MouseEvent.IsMouseButtonDown(EKeys::LeftMouseButton))
+					{
+						const TSharedRef<FSetCollectionDragDropOp> DragDropOp = FSetCollectionDragDropOp::New(InItem->Name, InItem->ParentName);
+						return FReply::Handled().BeginDragDrop(DragDropOp);
+					}
+
+					return FReply::Unhandled();
+				})
+			.OnCanAcceptDrop_Lambda([](const FDragDropEvent& DragDropEvent, EItemDropZone DropZone, TSharedPtr<FTreeItem> TargetTreeItem) -> TOptional<EItemDropZone>
+				{
+					using namespace UE::Anim::SkelTemplates;
+
+					TOptional<EItemDropZone> ReturnedDropZone;
+
+					const TSharedPtr<FSetCollectionDragDropOp> DragDropOp = DragDropEvent.GetOperationAs<FSetCollectionDragDropOp>();
+					if (DragDropOp.IsValid())
+					{
+						ReturnedDropZone = EItemDropZone::OntoItem;
+					}
+
+					return ReturnedDropZone;
+				})
+			.OnAcceptDrop_Lambda([this](const FDragDropEvent& DragDropEvent, EItemDropZone DropZone, TSharedPtr<FTreeItem> TargetTreeItem) -> FReply
+				{
+					const TSharedPtr<FSetCollectionDragDropOp> DragDropOp = DragDropEvent.GetOperationAs<FSetCollectionDragDropOp>();
+					if (DragDropOp.IsValid())
+					{
+						const FScopedTransaction Transaction(LOCTEXT("ReparentSet", "Reparent Set"));
+						SetCollection->Modify();
+
+						SetCollection->ReparentSet(DragDropOp->SetName, TargetTreeItem->Name);
+						RepopulateTreeData();
+
+						return FReply::Handled();
+					}
+
+					return FReply::Unhandled();
+				})
+			.ShowWires(true)
+			[
+				SNew(SHorizontalBox)
+					+ SHorizontalBox::Slot()
+					.AutoWidth()
+					.Padding(3.0f)
+					[
+						SNew(SImage)
+							.Image(FAppStyle::Get().GetBrush("ClassIcon.GroupActor"))
+					]
+					+ SHorizontalBox::Slot()
+					.FillWidth(1.0)
+					.AutoWidth()
+					.Padding(3.0f)
+					[
+						SAssignNew(InlineWidget, SInlineEditableTextBlock)
+							.Text(FText::FromName(InItem->Name))
+							.OnVerifyTextChanged_Lambda([InItem](const FText& InText, FText& OutErrorMessage)
+							{
+								const FName OldName = InItem->Name;
+								const FName NewName = FName(InText.ToString());
+
+								if (InItem->Name == NewName)
+								{
+									return true;
+								}
+
+								check(InItem->CanRenameTo.IsBound());
+								if (InItem->CanRenameTo.Execute(OldName, NewName))
+								{
+									return true;
+								}
+								else
+								{
+									OutErrorMessage = LOCTEXT("SetNameTaken", "Name already taken");
+									return false;
+								}
+							})
+							.OnTextCommitted_Lambda([InItem](const FText& InText, ETextCommit::Type CommitInfo)
+							{
+								const FName OldName = InItem->Name;
+								const FName NewName = FName(InText.ToString());
+
+								InItem->OnRenamed.ExecuteIfBound(OldName, NewName);
+							})
+					]
+
+			];
+
+		InItem->OnRequestRename.BindSP(InlineWidget.Get(), &SInlineEditableTextBlock::EnterEditingMode);
+
+		return RowWidget;
+	}
+
+	void SSetCollection::TreeView_OnGetChildren(TSharedPtr<FTreeItem> InItem, TArray<TSharedPtr<FTreeItem>>& OutChildren)
+	{
+		OutChildren = InItem->Children;
+	}
+
+	TSharedPtr<SWidget> SSetCollection::TreeView_OnContextMenuOpening()
+	{
+		FMenuBuilder MenuBuilder(true, CommandList.ToSharedRef());
+
+		MenuBuilder.AddMenuEntry(FGenericCommands::Get().Delete);
+
+		return MenuBuilder.MakeWidget();
+	}
+
+	void SSetCollection::TreeView_OnItemScrolledIntoView(TSharedPtr<FTreeItem> InItem, const TSharedPtr<ITableRow>& InWidget)
+	{
+		if (ItemToRename)
+		{
+			ItemToRename->OnRequestRename.ExecuteIfBound();
+			ItemToRename = nullptr;
+		}
+	}
+
+	void SSetCollection::TreeView_OnMouseButtonDoubleClick(TSharedPtr<FTreeItem> InItem)
+	{
+		InItem->OnRequestRename.ExecuteIfBound();
+	}
+
+	void SSetCollection::PostUndo(bool bSuccess)
+	{
+		RepopulateTreeData();
+	}
+
+	void SSetCollection::PostRedo(bool bSuccess)
+	{
+		RepopulateTreeData();
+	}
+
+	void SSetCollection::OnRemoveSet()
+	{
+		const TArray<TSharedPtr<FTreeItem>> SelectedItems = TreeView->GetSelectedItems();
+		if (!SelectedItems.IsEmpty())
+		{
+			const FScopedTransaction Transaction(LOCTEXT("RemoveSet", "Remove Set"));
+			SetCollection->Modify();
+
+			const FName SetToRemove = SelectedItems[0]->Name;
+			const bool bSuccess = SetCollection->RemoveSet(SetToRemove);
+			if (bSuccess)
+			{
+				RepopulateTreeData();
+			}
+		}
+	}
+
+	void SSetCollection::RepopulateTreeData()
+	{
+		// Make note of the currently selected item, if any
+		TArray<TSharedPtr<FTreeItem>> Selected = TreeView->GetSelectedItems();
+
+		// Make note of all tree items expanded
+		TArray<FName> ExpandedSetNames;
+
+		for (TSharedPtr<FTreeItem> TreeItem : GetAllTreeItems())
+		{
+			if (TreeView->IsItemExpanded(TreeItem))
+			{
+				ExpandedSetNames.Add(TreeItem->Name);
+			}
+		}
+
+		// Rebuild tree data
+		RootItems.Reset();
+
+		TMap<FName, TSharedPtr<FTreeItem>> TreeItemMap;
+
+		for (const FAbstractSkeletonSet& Set : SetCollection->GetSetHierarchy())
+		{
+			TSharedPtr<FTreeItem> SetItem = MakeShared<FTreeItem>();
+			SetItem->Name = Set.SetName;
+			SetItem->ParentName = Set.ParentSetName;
+			SetItem->OnRenamed = FTreeItem::FOnRenamed::CreateSP(this, &SSetCollection::OnItemRenamed);
+			SetItem->CanRenameTo = FTreeItem::FCanRenameTo::CreateSP(this, &SSetCollection::CanRenameItem);
+
+			if (Set.ParentSetName == NAME_None)
+			{
+				RootItems.Add(SetItem);
+			}
+			else
+			{
+				if (!TreeItemMap.Contains(Set.ParentSetName))
+				{
+					TSharedPtr<FTreeItem> ParentItem = MakeShared<FTreeItem>();
+					ParentItem->Name = Set.ParentSetName;
+					ParentItem->OnRenamed = FTreeItem::FOnRenamed::CreateSP(this, &SSetCollection::OnItemRenamed);
+					ParentItem->CanRenameTo = FTreeItem::FCanRenameTo::CreateSP(this, &SSetCollection::CanRenameItem);
+
+					TreeItemMap.Add(Set.ParentSetName, ParentItem);
+				}
+
+				TSharedPtr<FTreeItem> ParentItem = TreeItemMap[Set.ParentSetName];
+				ParentItem->Children.Add(SetItem);
+
+			}
+
+			TreeItemMap.Add(Set.SetName, SetItem);
+		}
+
+		// Restore expanded state of tree items and selection
+		const FName SelectedSetName = Selected.IsEmpty() ? NAME_None : Selected[0]->Name;
+		for (TSharedPtr<FTreeItem> TreeItem : GetAllTreeItems())
+		{
+			if (ExpandedSetNames.Contains(TreeItem->Name))
+			{
+				TreeView->SetItemExpansion(TreeItem, true);
+			}
+
+			if (TreeItem->Name == SelectedSetName)
+			{
+				TreeView->SetSelection(TreeItem);
+			}
+		}
+	}
+
+	void SSetCollection::ExpandAllTreeItems()
+	{
+		for (TSharedPtr<FTreeItem> TreeItem : GetAllTreeItems())
+		{
+			TreeView->SetItemExpansion(TreeItem, true);
+		}
+	}
+
+	void SSetCollection::BindCommands()
+	{
+		CommandList = MakeShareable(new FUICommandList);
+
+		CommandList->MapAction(
+			FGenericCommands::Get().Delete,
+			FExecuteAction::CreateSP(this, &SSetCollection::OnRemoveSet)
+		);
+	}
+
+	TArray<TSharedPtr<SSetCollection::FTreeItem>> SSetCollection::GetAllTreeItems()
+	{
+		TArray<TSharedPtr<FTreeItem>> AllItems;
+		AllItems.Append(RootItems);
+
+		for (int32 Index = 0; Index < AllItems.Num(); ++Index)
+		{
+			AllItems.Append(AllItems[Index]->Children);
+		}
+
+		return AllItems;
+	}
+
+	FReply SSetCollection::OnAddSet()
+	{
+		const FScopedTransaction Transaction(LOCTEXT("AddSet", "Set Set"));
+		SetCollection->Modify();
+
+		TArray<TSharedPtr<FTreeItem>> SelectedItems = TreeView->GetSelectedItems();
+
+		const FName ParentSetName = SelectedItems.IsEmpty() ? NAME_None : SelectedItems[0]->Name;
+		
+		const FText DefaultSetNameFormat = LOCTEXT("DefaultSetName", "NewSet_{0}");
+
+		int32 Suffix = 0;
+		FName SetName = FName(FText::Format(DefaultSetNameFormat, Suffix).ToString());
+		while (SetCollection->HasSet(SetName))
+		{
+			++Suffix;
+			SetName = FName(FText::Format(DefaultSetNameFormat, Suffix).ToString());
+		}
+
+		SetCollection->AddSet(SetName, ParentSetName);
+		RepopulateTreeData();
+
+		// Expand the parent set tree item so we can see the new set in the tree
+		TreeView->SetItemExpansion(GetTreeItem(ParentSetName), true);
+
+		TSharedPtr<FTreeItem> SetTreeItem = GetTreeItem(SetName);
+		ItemToRename = SetTreeItem;
+		TreeView->RequestScrollIntoView(SetTreeItem);
+
+		return FReply::Handled();
+	}
+
+	void SSetCollection::OnItemRenamed(const FName OldSetName, const FName NewSetName)
+	{
+		SetCollection->RenameSet(OldSetName, NewSetName);
+		RepopulateTreeData();
+	}
+
+	bool SSetCollection::CanRenameItem(const FName OldSetName, const FName NewSetName) const
+	{
+		return !SetCollection->HasSet(NewSetName);
+	}
+
+	TSharedPtr<SSetCollection::FTreeItem> SSetCollection::GetTreeItem(const FName SetName)
+	{
+		for (TSharedPtr<FTreeItem> TreeItem : GetAllTreeItems())
+		{
+			if (TreeItem->Name == SetName)
+			{
+				return TreeItem;
+			}
+		}
+
+		return nullptr;
+	}
+}
+
+#undef LOCTEXT_NAMESPACE
\ No newline at end of file
diff --git a/Engine/Plugins/Experimental/Animation/SkeletonTemplateFramework/Source/Editor/Private/SkeletonTemplateFramework/Sets/SetCollectionEditor.h b/Engine/Plugins/Experimental/Animation/SkeletonTemplateFramework/Source/Editor/Private/SkeletonTemplateFramework/Sets/SetCollectionEditor.h
new file mode 100644
--- /dev/null
+++ b/Engine/Plugins/Experimental/Animation/SkeletonTemplateFramework/Source/Editor/Private/SkeletonTemplateFramework/Sets/SetCollectionEditor.h
@@ -0,0 +1,107 @@
+// Copyright Epic Games, Inc. All Rights Reserved.
+
+#pragma once
+
+#include "CoreMinimal.h"
+#include "EditorUndoClient.h"
+#include "SkeletonTemplateFramework/AbstractSkeletonSetCollection.h"
+#include "Toolkits/AssetEditorToolkit.h"
+#include "Widgets/SCompoundWidget.h"
+#include "Widgets/Views/STreeView.h"
+
+namespace UE::Anim::SkelTemplates
+{
+	class FSetCollectionEditorToolkit : public FAssetEditorToolkit
+	{
+	public:
+		void InitEditor(const TArray<UObject*>& InObjects);
+
+		// Begin FAssetEditorToolkit
+		void RegisterTabSpawners(const TSharedRef<class FTabManager>& TabManager) override;
+		void UnregisterTabSpawners(const TSharedRef<class FTabManager>& TabManager) override;
+
+		FName GetToolkitFName() const override { return "AbstractSkeletonSetCollectionEditor"; }
+		FText GetBaseToolkitName() const override { return INVTEXT("Abstract Skeleton Set Collection Editor"); }
+		FString GetWorldCentricTabPrefix() const override { return "Abstract Skeleton Set Collection Editor "; }
+		FLinearColor GetWorldCentricTabColorScale() const override { return {}; }
+		// End FAssetEditorToolkit
+
+	private:
+		TSharedRef<SDockTab> SpawnPropertiesTab(const FSpawnTabArgs& Args);
+
+		TWeakObjectPtr<UAbstractSkeletonSetCollection> SetCollection;
+	};
+
+	class SSetCollection : public SCompoundWidget, public FSelfRegisteringEditorUndoClient
+	{
+		SLATE_BEGIN_ARGS(SSetCollection) {}
+		SLATE_END_ARGS()
+
+	public:
+		struct FTreeItem
+		{
+			FName Name;
+			FName ParentName;
+			TArray<TSharedPtr<FTreeItem>> Children;
+
+			DECLARE_DELEGATE(FOnRequestRename);
+			FOnRequestRename OnRequestRename;
+
+			DECLARE_DELEGATE_RetVal_TwoParams(bool, FCanRenameTo, const FName /* Old Name */, const FName /* New Name */);
+			FCanRenameTo CanRenameTo;
+
+			DECLARE_DELEGATE_TwoParams(FOnRenamed, const FName /* Old Name */, const FName /* New Name */);
+			FOnRenamed OnRenamed;
+		};
+
+		void Construct(const FArguments& InArgs, TWeakObjectPtr<UAbstractSkeletonSetCollection> InSetCollection);
+
+		/** FSelfRegisteringEditorUndoClient interface */
+		virtual void PostUndo(bool bSuccess) override;
+		virtual void PostRedo(bool bSuccess) override;
+
+	private:
+		virtual FReply OnKeyDown( const FGeometry& MyGeometry, const FKeyEvent& InKeyEvent) override;
+
+		TSharedRef<ITableRow> TreeView_OnGenerateRow(TSharedPtr<FTreeItem> InItem, const TSharedRef<STableViewBase>& OwnerTable);
+
+		void TreeView_OnGetChildren(TSharedPtr<FTreeItem> InItem, TArray<TSharedPtr<FTreeItem>>& OutChildren);
+
+		TSharedPtr<SWidget> TreeView_OnContextMenuOpening();
+
+		void TreeView_OnItemScrolledIntoView(TSharedPtr<FTreeItem> InItem, const TSharedPtr<ITableRow>& InWidget);
+		
+		void TreeView_OnMouseButtonDoubleClick(TSharedPtr<FTreeItem> InItem);
+
+		void OnRemoveSet();
+
+		void RepopulateTreeData();
+
+		void ExpandAllTreeItems();
+
+		void BindCommands();
+
+		TArray<TSharedPtr<FTreeItem>> GetAllTreeItems();
+
+		FReply OnAddSet();
+
+		void OnItemRenamed(const FName OldSetName, const FName NewSetName);
+
+		bool CanRenameItem(const FName OldSetName, const FName NewSetName) const;
+
+	private:
+		TSharedPtr<FTreeItem> GetTreeItem(const FName SetName);
+
+		TWeakObjectPtr<UAbstractSkeletonSetCollection> SetCollection;
+
+		TSharedPtr<STreeView<TSharedPtr<FTreeItem>>> TreeView;
+
+		TArray<TSharedPtr<FTreeItem>> RootItems;
+
+		TSharedPtr<SSearchBox> SearchBox;
+
+		TSharedPtr<FUICommandList> CommandList;
+
+		TSharedPtr<FTreeItem> ItemToRename;
+	};
+}
diff --git a/Engine/Plugins/Experimental/Animation/SkeletonTemplateFramework/Source/Editor/Private/SkeletonTemplateFramework/Sets/SetCollectionFactory.cpp b/Engine/Plugins/Experimental/Animation/SkeletonTemplateFramework/Source/Editor/Private/SkeletonTemplateFramework/Sets/SetCollectionFactory.cpp
new file mode 100644
--- /dev/null
+++ b/Engine/Plugins/Experimental/Animation/SkeletonTemplateFramework/Source/Editor/Private/SkeletonTemplateFramework/Sets/SetCollectionFactory.cpp
@@ -0,0 +1,31 @@
+// Copyright Epic Games, Inc. All Rights Reserved.
+
+#include "SkeletonTemplateFramework/Sets/SetCollectionFactory.h"
+
+#include "SkeletonTemplateFramework/AbstractSkeletonSetCollection.h"
+#include "UObject/Package.h"
+
+UAbstractSkeletonSetCollectionFactory::UAbstractSkeletonSetCollectionFactory()
+{
+	bCreateNew = true;
+	bEditAfterNew = true;
+	SupportedClass = UAbstractSkeletonSetCollection::StaticClass();
+}
+
+bool UAbstractSkeletonSetCollectionFactory::ConfigureProperties()
+{
+	return true;
+}
+
+UObject* UAbstractSkeletonSetCollectionFactory::FactoryCreateNew(UClass* Class, UObject* InParent, FName Name, EObjectFlags Flags, UObject* Context, FFeedbackContext* Warn, FName CallingContext)
+{
+	EObjectFlags FlagsToUse = Flags | RF_Public | RF_Standalone | RF_Transactional | RF_LoadCompleted;
+	if (InParent == GetTransientPackage())
+	{
+		FlagsToUse &= ~RF_Standalone;
+	}
+
+	TObjectPtr<UAbstractSkeletonSetCollection> NewSetCollection = NewObject<UAbstractSkeletonSetCollection>(InParent, Class, Name, FlagsToUse);
+
+	return NewSetCollection;
+}
diff --git a/Engine/Plugins/Experimental/Animation/SkeletonTemplateFramework/Source/Editor/Private/SkeletonTemplateFramework/Sets/SetCollectionFactory.h b/Engine/Plugins/Experimental/Animation/SkeletonTemplateFramework/Source/Editor/Private/SkeletonTemplateFramework/Sets/SetCollectionFactory.h
new file mode 100644
--- /dev/null
+++ b/Engine/Plugins/Experimental/Animation/SkeletonTemplateFramework/Source/Editor/Private/SkeletonTemplateFramework/Sets/SetCollectionFactory.h
@@ -0,0 +1,23 @@
+// Copyright Epic Games, Inc. All Rights Reserved.
+
+#pragma once
+
+#include "CoreMinimal.h"
+#include "Factories/Factory.h"
+
+#include "SetCollectionFactory.generated.h"
+
+UCLASS()
+class UAbstractSkeletonSetCollectionFactory : public UFactory
+{
+	GENERATED_BODY()
+
+	UAbstractSkeletonSetCollectionFactory();
+
+	virtual bool ConfigureProperties() override;
+	virtual UObject* FactoryCreateNew(UClass* Class, UObject* InParent, FName Name, EObjectFlags Flags, UObject* Context, FFeedbackContext* Warn, FName CallingContext) override;
+	virtual UObject* FactoryCreateNew(UClass* Class, UObject* InParent, FName Name, EObjectFlags Flags, UObject* Context, FFeedbackContext* Warn) override
+	{
+		return FactoryCreateNew(Class, InParent, Name, Flags, Context, Warn, NAME_None);
+	}
+};
diff --git a/Engine/Plugins/Experimental/Animation/SkeletonTemplateFramework/Source/Editor/Private/SkeletonTemplateFramework/Sets/SetDragDrop.cpp b/Engine/Plugins/Experimental/Animation/SkeletonTemplateFramework/Source/Editor/Private/SkeletonTemplateFramework/Sets/SetDragDrop.cpp
new file mode 100644
--- /dev/null
+++ b/Engine/Plugins/Experimental/Animation/SkeletonTemplateFramework/Source/Editor/Private/SkeletonTemplateFramework/Sets/SetDragDrop.cpp
@@ -0,0 +1,27 @@
+// Copyright Epic Games, Inc. All Rights Reserved.
+
+#include "SkeletonTemplateFramework/Sets/SetDragDrop.h"
+
+namespace UE::Anim::SkelTemplates::Sets
+{
+	TSharedRef<FSetDragDropOp> FSetDragDropOp::New(SSetBinding::FTreeItem_Set* InItem)
+	{
+		TSharedRef<FSetDragDropOp> Operation = MakeShared<FSetDragDropOp>();
+		Operation->Item = InItem;
+
+		Operation->Construct();
+
+		return Operation;
+	}
+
+	TSharedPtr<SWidget> FSetDragDropOp::GetDefaultDecorator() const
+	{
+		return SNew(SBorder)
+			.Visibility(EVisibility::Visible)
+			.BorderImage(FAppStyle::GetBrush("Menu.Background"))
+			[
+				SNew(STextBlock)
+					.Text(FText::FromName(Item->Name))
+			];
+	}
+}
\ No newline at end of file
diff --git a/Engine/Plugins/Experimental/Animation/SkeletonTemplateFramework/Source/Editor/Private/SkeletonTemplateFramework/Sets/SetDragDrop.h b/Engine/Plugins/Experimental/Animation/SkeletonTemplateFramework/Source/Editor/Private/SkeletonTemplateFramework/Sets/SetDragDrop.h
new file mode 100644
--- /dev/null
+++ b/Engine/Plugins/Experimental/Animation/SkeletonTemplateFramework/Source/Editor/Private/SkeletonTemplateFramework/Sets/SetDragDrop.h
@@ -0,0 +1,22 @@
+// Copyright Epic Games, Inc. All Rights Reserved.
+
+#pragma once
+
+#include "CoreMinimal.h"
+#include "DragAndDrop/DecoratedDragDropOp.h"
+#include "SkeletonTemplateFramework/Sets/SSetBinding.h"
+
+namespace UE::Anim::SkelTemplates::Sets
+{
+	class FSetDragDropOp : public FDecoratedDragDropOp
+	{
+	public:
+		DRAG_DROP_OPERATOR_TYPE(FSetDragDropOp, FDecoratedDragDropOp)
+
+		static TSharedRef<FSetDragDropOp> New(SSetBinding::FTreeItem_Set* InItem);
+
+		TSharedPtr<SWidget> GetDefaultDecorator() const override;
+
+		SSetBinding::FTreeItem_Set* Item;
+	};
+}
diff --git a/Engine/Plugins/Experimental/Animation/SkeletonTemplateFramework/Source/Editor/Private/SkeletonTemplateFramework/Sets/SkeletonTreeDragDrop.h b/Engine/Plugins/Experimental/Animation/SkeletonTemplateFramework/Source/Editor/Private/SkeletonTemplateFramework/Sets/SkeletonTreeDragDrop.h
new file mode 100644
--- /dev/null
+++ b/Engine/Plugins/Experimental/Animation/SkeletonTemplateFramework/Source/Editor/Private/SkeletonTemplateFramework/Sets/SkeletonTreeDragDrop.h
@@ -0,0 +1,39 @@
+// Copyright Epic Games, Inc. All Rights Reserved.
+
+#pragma once
+
+#include "CoreMinimal.h"
+#include "SSetSkeletonTree.h"
+#include "DragAndDrop/DecoratedDragDropOp.h"
+
+namespace UE::Anim::SkelTemplates::Sets
+{
+	class FSkeletonTreeDragDropOp : public FDecoratedDragDropOp
+	{
+	public:
+		DRAG_DROP_OPERATOR_TYPE(FSkeletonTreeDragDropOp, FDecoratedDragDropOp)
+
+		static TSharedRef<FSkeletonTreeDragDropOp> New(TSharedPtr<SSetsSkeletonTree::FTreeItem> InItem)
+		{
+			TSharedRef<FSkeletonTreeDragDropOp> Operation = MakeShared<FSkeletonTreeDragDropOp>();
+			Operation->Item = InItem;
+
+			Operation->Construct();
+
+			return Operation;
+		}
+
+		TSharedPtr<SWidget> GetDefaultDecorator() const override
+		{
+			return SNew(SBorder)
+				.Visibility(EVisibility::Visible)
+				.BorderImage(FAppStyle::GetBrush("Menu.Background"))
+				[
+					SNew(STextBlock)
+						.Text(FText::FromName(Item->BoneName))
+				];
+		}
+
+		TSharedPtr<SSetsSkeletonTree::FTreeItem> Item;
+	};
+}
\ No newline at end of file
diff --git a/Engine/Plugins/Experimental/Animation/SkeletonTemplateFramework/Source/Editor/SkeletonTemplateFrameworkEditor.Build.cs b/Engine/Plugins/Experimental/Animation/SkeletonTemplateFramework/Source/Editor/SkeletonTemplateFrameworkEditor.Build.cs
--- a/Engine/Plugins/Experimental/Animation/SkeletonTemplateFramework/Source/Editor/SkeletonTemplateFrameworkEditor.Build.cs
+++ b/Engine/Plugins/Experimental/Animation/SkeletonTemplateFramework/Source/Editor/SkeletonTemplateFrameworkEditor.Build.cs
@@ -18,7 +18,11 @@ public class SkeletonTemplateFrameworkEditor : ModuleRules
 			"UnrealEd",
 			"SkeletonTemplateFrameworkRuntime",
 			"Slate",
-			"SlateCore"
+			"SlateCore",
+			"Persona",
+			"PropertyEditor",
+			"ToolWidgets",
+			"ToolMenus"
 		);
 	}
 }
