Fix abstract skeleton compile error

#rnx

diff --git a/Engine/Plugins/Experimental/Animation/SkeletonTemplateFramework/Source/Editor/Private/SkeletonTemplateFramework/Sets/SSetSkeletonTree.cpp b/Engine/Plugins/Experimental/Animation/SkeletonTemplateFramework/Source/Editor/Private/SkeletonTemplateFramework/Sets/SSetSkeletonTree.cpp
--- a/Engine/Plugins/Experimental/Animation/SkeletonTemplateFramework/Source/Editor/Private/SkeletonTemplateFramework/Sets/SSetSkeletonTree.cpp
+++ b/Engine/Plugins/Experimental/Animation/SkeletonTemplateFramework/Source/Editor/Private/SkeletonTemplateFramework/Sets/SSetSkeletonTree.cpp
@@ -240,27 +240,27 @@ namespace UE::Anim::SkelTemplates::Sets
 			{
 				const FReferenceSkeleton& RefSkeleton = SetBinding->GetSkeleton()->GetReferenceSkeleton();
 
-				using BoneFunctionType = TFunctionRef<TSharedPtr<FTreeItem>(const int32, TSharedPtr<FTreeItem>)>;
+				TFunction<TSharedPtr<FTreeItem>(int32, const TSharedPtr<FTreeItem>&)> AddBoneToTree;
 
-				const BoneFunctionType AddBoneToTree = [&, this](const int32 BoneIndex, TSharedPtr<FTreeItem> Parent) -> TSharedPtr<FTreeItem>
-					{
-						TSharedPtr<FTreeItem> TreeItem = MakeShared<FTreeItem>();
-						TreeItem->BoneName = RefSkeleton.GetBoneName(BoneIndex);
-						TreeItem->BoundSet = SetBinding->GetBoneSet(TreeItem->BoneName);
-						TreeItem->Parent = Parent;
-
-						TArray<int32> Children;
-						RefSkeleton.GetDirectChildBones(BoneIndex, Children);
+				AddBoneToTree = [this, &RefSkeleton, &AddBoneToTree](const int32 BoneIndex, const TSharedPtr<FTreeItem>& Parent) -> TSharedPtr<FTreeItem>
+				{
+					TSharedPtr<FTreeItem> TreeItem = MakeShared<FTreeItem>();
+					TreeItem->BoneName = RefSkeleton.GetBoneName(BoneIndex);
+					TreeItem->BoundSet = SetBinding->GetBoneSet(TreeItem->BoneName);
+					TreeItem->Parent = Parent;
 
-						for (const int32 ChildBoneIndex : Children)
-						{
-							TSharedPtr<FTreeItem> ChildTreeItem = AddBoneToTree(ChildBoneIndex, TreeItem);
-							TreeItem->Children.Add(ChildTreeItem);
-						}
+					TArray<int32> Children;
+					RefSkeleton.GetDirectChildBones(BoneIndex, Children);
 
-						return TreeItem;
-					};
+					for (const int32 ChildBoneIndex : Children)
+					{
+						TSharedPtr<FTreeItem> ChildTreeItem = AddBoneToTree(ChildBoneIndex, TreeItem);
+						TreeItem->Children.Add(ChildTreeItem);
+					}
 
+					return TreeItem;
+				};
+				
 				TSharedPtr<FTreeItem> RootBoneItem = AddBoneToTree(0, nullptr);
 
 				RootItems.Add(RootBoneItem);
