Fix LOCTEXT conflict

#rnx

diff --git a/Engine/Plugins/Experimental/Animation/SkeletonTemplateFramework/Source/Editor/Private/SkeletonTemplateFramework/AbstractSkeletonEditor.cpp b/Engine/Plugins/Experimental/Animation/SkeletonTemplateFramework/Source/Editor/Private/SkeletonTemplateFramework/AbstractSkeletonEditor.cpp
--- a/Engine/Plugins/Experimental/Animation/SkeletonTemplateFramework/Source/Editor/Private/SkeletonTemplateFramework/AbstractSkeletonEditor.cpp
+++ b/Engine/Plugins/Experimental/Animation/SkeletonTemplateFramework/Source/Editor/Private/SkeletonTemplateFramework/AbstractSkeletonEditor.cpp
@@ -282,7 +282,7 @@ namespace UE::Anim::SkelTemplates
 
 	FText FAbstractSkeletonEditor::GetToolkitToolTipText() const
 	{
-		return LOCTEXT("ToolkitName", "Editor for manipulating a Skeleton's Set Bindings and Label Bindings");
+		return LOCTEXT("ToolkitTooltip", "Editor for manipulating a Skeleton's Set Bindings and Label Bindings");
 	}
 
 	TSharedPtr<IPersonaToolkit> FAbstractSkeletonEditor::GetPersonaToolkit() const
